AnimGen: Move to Experimental

diff --git a/Engine/Plugins/Experimental/Animation/AnimGen/Source/AnimGen/Private/AnimGenBehavior.cpp b/Engine/Plugins/Experimental/Animation/AnimGen/Source/AnimGen/Private/AnimGenBehavior.cpp
new file mode 100644
--- /dev/null
+++ b/Engine/Plugins/Experimental/Animation/AnimGen/Source/AnimGen/Private/AnimGenBehavior.cpp
@@ -0,0 +1,1993 @@
+// Copyright Epic Games, Inc. All Rights Reserved.
+
+#include "AnimGenBehavior.h"
+#include "AnimDatabasePose.h"
+#include "AnimDatabaseMath.h"
+
+#include "LearningRandom.h"
+
+#include "AnimGenLog.h"
+#include "InputCoreTypes.h"
+#include "Animation/SpringMath.h"
+#include "DrawDebugLibrary.h"
+
+#include "GameFramework/Actor.h"
+
+#include UE_INLINE_GENERATED_CPP_BY_NAME(AnimGenBehavior)
+
+#if WITH_EDITOR
+
+FAnimGenControlSchemaElement UAnimGenBehavior::SpecifyControl_Implementation(FAnimGenControlSchema& InControlSchema) const
+{
+	UE_LOGFMT(LogAnimGen, Error, "SpecifyControl must be implemented by Behavior.");
+	return FAnimGenControlSchemaElement();
+}
+
+void UAnimGenBehavior::DrawDebugControl_Implementation(
+	const FDebugDrawer& Drawer,
+	const FAnimGenControlObject& InControlObject,
+	const FAnimGenControlObjectElement InControlObjectElement,
+	const FAnimDatabasePoseState& InPoseState,
+	const UAnimDatabase* InDatabase,
+	const UAnimGenAutoEncoder* InAutoEncoder,
+	const int32 SequenceIdx,
+	const float SequenceTime,
+	const int32 RangeStart,
+	const int32 RangeLength,
+	const FLinearColor& IdentifierColor) const {}
+
+void UAnimGenBehavior::MakeControlSets_Implementation(TArray<FAnimGenControlSet>& OutControlSets, FAnimGenControlObject& InControlObject, const UAnimDatabase* Database, const FAnimDatabaseFrameRanges& InFrameRanges) const
+{
+	UE_LOGFMT(LogAnimGen, Error, "MakeControlSets must be implemented by Behavior.");
+	OutControlSets.Empty();
+}
+
+void UAnimGenBehaviorPreview::OnPreviewInitialize_Implementation(const UAnimGenController* Controller, const UAnimGenAutoEncoder* AutoEncoder) {}
+
+FAnimGenControlObjectElement UAnimGenBehaviorPreview::MakeControl_Implementation(FAnimGenControlObject& InControlObject, const FAnimDatabasePoseState& InPoseState) const
+{
+	return FAnimGenControlObjectElement();
+}
+
+bool UAnimGenBehaviorPreview::OnPreviewCheckReset_Implementation(FAnimGenBehaviorPreviewReset& OutResetData, const FAnimDatabasePoseState& State) { return false; }
+void UAnimGenBehaviorPreview::OnPreviewReset_Implementation(const FAnimDatabasePoseState& State) {}
+void UAnimGenBehaviorPreview::OnPreviewUpdate_Implementation(const FAnimDatabasePoseState& State, const float DeltaTime) {}
+
+void UAnimGenBehaviorPreview::OnPreviewDrawDebug_Implementation(
+	const FDebugDrawer& Drawer, 
+	const FAnimGenControlObject& InControlObject,
+	const FAnimGenControlObjectElement InControlObjectElement, 
+	const FAnimDatabasePoseState& State, 
+	const FLinearColor& IdentifierColor) {}
+
+void UAnimGenBehaviorPreview::OnPreviewInputLeftStickAxis_Implementation(const FVector Axis) {}
+void UAnimGenBehaviorPreview::OnPreviewInputRightStickAxis_Implementation(const FVector Axis) {}
+void UAnimGenBehaviorPreview::OnPreviewInputLeftTriggerAxis_Implementation(const float Axis) {}
+void UAnimGenBehaviorPreview::OnPreviewInputRightTriggerAxis_Implementation(const float Axis) {}
+void UAnimGenBehaviorPreview::OnPreviewInputMouseMoved_Implementation(const FVector MouseGroundLocation) {}
+void UAnimGenBehaviorPreview::OnPreviewInputLeftMousePressed_Implementation(const FVector MouseGroundLocation) {}
+void UAnimGenBehaviorPreview::OnPreviewInputLeftMouseReleased_Implementation(const FVector MouseGroundLocation) {}
+void UAnimGenBehaviorPreview::OnPreviewInputRightMousePressed_Implementation(const FVector MouseGroundLocation) {}
+void UAnimGenBehaviorPreview::OnPreviewInputRightMouseReleased_Implementation(const FVector MouseGroundLocation) {}
+void UAnimGenBehaviorPreview::OnPreviewInputKeyPressed_Implementation(const FKey Key) {}
+void UAnimGenBehaviorPreview::OnPreviewInputKeyReleased_Implementation(const FKey Key) {}
+
+void UAnimGenBehaviorPreview::GetPreviewCameraLocationAndRotation(FVector& OutCameraLocation, FRotator& OutCameraRotation) const
+{
+	OutCameraLocation = PreviewCameraLocation;
+	OutCameraRotation = PreviewCameraRotation;
+}
+
+FRotator UAnimGenBehaviorPreview::GetPreviewCameraYawRotation() const
+{
+	return FRotator(0.0f, PreviewCameraRotation.Yaw, 0.0f);
+}
+
+void UAnimGenBehaviorPreview::UpdatePreviewCamera(
+	const float YawVelocity,
+	const float PitchVelocity,
+	const float DistanceVelocity,
+	const float DeltaTime,
+	const bool bInvertPitch)
+{
+	PreviewCameraYaw += DeltaTime * YawVelocity;
+	PreviewCameraPitch = FMath::Clamp(PreviewCameraPitch + DeltaTime * (bInvertPitch ? -PitchVelocity : PitchVelocity), 0.0f, 80.0f);
+	PreviewCameraDistance = FMath::Clamp(PreviewCameraDistance + DeltaTime * DistanceVelocity, 10.0f, 1000.0);
+}
+
+void UAnimGenBehaviorPreview::UpdatePreviewCameraFromAxisAndShoulders(
+	const FVector Axis,
+	const bool bLeftShoulder,
+	const bool bRightShoulder,
+	const float DeltaTime,
+	const float RotationSpeed,
+	const float DistanceSpeed,
+	const bool bInvertPitch,
+	const float Deadzone,
+	const float Sensitivity)
+{
+	FVector OutAxis = FVector::ZeroVector;
+	FVector OutDirection = FVector::ZeroVector;
+	float OutLength = 0.0f;
+	bool bOutNotInDeadzone = false;
+	UAnimGenBehaviorLibrary::ApplyStickDeadzoneAndSensitivity(
+		OutAxis,
+		OutDirection,
+		OutLength,
+		bOutNotInDeadzone,
+		Axis,
+		Deadzone,
+		Sensitivity);
+
+	UpdatePreviewCamera(
+		bOutNotInDeadzone ? RotationSpeed * -OutAxis.Y : 0.0f,
+		bOutNotInDeadzone ? RotationSpeed * OutAxis.X : 0.0f,
+		bLeftShoulder ? DistanceSpeed : bRightShoulder ? -DistanceSpeed : 0.0f,
+		DeltaTime,
+		bInvertPitch);
+}
+
+FVector UAnimGenBehaviorPreview::GetPreviewLeftStickAxis() const
+{
+	return PreviewLeftStick;
+}
+
+FVector UAnimGenBehaviorPreview::GetPreviewRightStickAxis() const
+{
+	return PreviewRightStick;
+}
+
+FVector UAnimGenBehaviorPreview::GetPreviewCameraRelativeLeftStickAxis() const
+{
+	return GetPreviewCameraYawRotation().RotateVector(PreviewLeftStick);
+}
+
+FVector UAnimGenBehaviorPreview::GetPreviewCameraRelativeRightStickAxis() const
+{
+	return GetPreviewCameraYawRotation().RotateVector(PreviewRightStick);
+}
+
+float UAnimGenBehaviorPreview::GetPreviewLeftTriggerAxis() const
+{
+	return PreviewLeftTrigger;
+}
+
+float UAnimGenBehaviorPreview::GetPreviewRightTriggerAxis() const
+{
+	return PreviewRightTrigger;
+}
+
+float UAnimGenBehaviorPreview::GetPreviewLeftTriggerAxisWithSensitivity(const float Sensitivity) const
+{
+	return UAnimGenBehaviorLibrary::ApplyTriggerSensitivity(PreviewLeftTrigger, Sensitivity);
+}
+
+float UAnimGenBehaviorPreview::GetPreviewRightTriggerAxisWithSensitivity(const float Sensitivity) const
+{
+	return UAnimGenBehaviorLibrary::ApplyTriggerSensitivity(PreviewRightTrigger, Sensitivity);
+}
+
+FVector UAnimGenBehaviorPreview::GetPreviewMouseGroundLocation() const
+{
+	return PreviewMouseGroundLocation;
+}
+
+FAnimGenControlSchemaElement UAnimGenBehavior_Uncontrolled::SpecifyControl_Implementation(FAnimGenControlSchema& InControlSchema) const
+{
+	return UAnimGenControls::SpecifyNullControl(InControlSchema);
+}
+
+void UAnimGenBehavior_Uncontrolled::MakeControlSets_Implementation(TArray<FAnimGenControlSet>& OutControlSets, FAnimGenControlObject& InControlObject, const UAnimDatabase* InDatabase, const FAnimDatabaseFrameRanges& InFrameRanges) const
+{
+	OutControlSets.Add(UAnimGenControls::MakeNullControlSet(InControlObject, InFrameRanges));
+}
+
+FAnimGenControlSchemaElement UAnimGenBehavior_Idle::SpecifyControl_Implementation(FAnimGenControlSchema& InControlSchema) const
+{
+	const int32 TagNum = TagRanges.Num();
+
+	TArray<FName, TInlineAllocator<32>> TagNames;
+	TagNames.Reserve(TagNum);
+
+	for (int32 TagIdx = 0; TagIdx < TagNum; TagIdx++)
+	{
+		TagNames.AddUnique(TagRanges[TagIdx].Name);
+	}
+
+	return UAnimGenControls::SpecifyNamedInclusiveDiscreteControlFromArrayView(InControlSchema, TagNames);
+}
+
+void UAnimGenBehavior_Idle::MakeControlSets_Implementation(TArray<FAnimGenControlSet>& OutControlSets, FAnimGenControlObject& InControlObject, const UAnimDatabase* InDatabase, const FAnimDatabaseFrameRanges& InFrameRanges) const
+{
+	const int32 TagNum = TagRanges.Num();
+
+	TArray<FName, TInlineAllocator<32>> TagNames;
+	TArray<FAnimDatabaseFrameRanges, TInlineAllocator<32>> TagFrameRanges;
+	TagNames.Reserve(TagNum);
+	TagFrameRanges.Reserve(TagNum);
+
+	for (int32 TagIdx = 0; TagIdx < TagNum; TagIdx++)
+	{
+		if (TagNames.Contains(TagRanges[TagIdx].Name)) { continue; }
+
+		const FAnimDatabaseFrameRanges Ranges = UAnimDatabaseFrameRangesLibrary::MakeFrameRangesFromFunction(InDatabase, InFrameRanges, TagRanges[TagIdx].FrameRanges);
+
+		if (Ranges.IsValid() && !Ranges.FrameRangeSet->IsEmpty())
+		{
+			TagNames.Add(TagRanges[TagIdx].Name);
+			TagFrameRanges.Add(Ranges);
+		}
+	}
+
+	OutControlSets.Add(UAnimGenControls::MakeNamedInclusiveDiscreteControlSetFromFrameRangesArrayViews(InControlObject, InFrameRanges, TagNames, TagFrameRanges));
+}
+
+void UAnimGenBehavior_Idle::DrawDebugControl_Implementation(
+	const FDebugDrawer& Drawer,
+	const FAnimGenControlObject& InControlObject,
+	const FAnimGenControlObjectElement InControlObjectElement,
+	const FAnimDatabasePoseState& InPoseState,
+	const UAnimDatabase* InDatabase,
+	const UAnimGenAutoEncoder* InAutoEncoder,
+	const int32 SequenceIdx,
+	const float SequenceTime,
+	const int32 RangeStart,
+	const int32 RangeLength,
+	const FLinearColor& IdentifierColor) const
+{
+	const FTransform RootTransform = UAnimDatabasePoseStateLibrary::PoseStateRootTransform(InPoseState);
+
+	UAnimGenControls::DrawDebugNamedInclusiveDiscreteControl(Drawer, InControlObject, InControlObjectElement, TEXT(""), RootTransform.GetLocation() + FVector(0.0, 0.0, 220.0f), FRotator(0.0f, -90.0f, 0.0f), 10.0f, IdentifierColor);
+}
+
+FAnimGenControlSchemaElement UAnimGenBehavior_TrajectoryFollow::SpecifyControl_Implementation(FAnimGenControlSchema& InControlSchema) const
+{
+	const int32 TagNum = TagRanges.Num();
+
+	TArray<FName, TInlineAllocator<32>> TagNames;
+	TagNames.Reserve(TagNum);
+
+	for (int32 TagIdx = 0; TagIdx < TagNum; TagIdx++)
+	{
+		TagNames.AddUnique(TagRanges[TagIdx].Name);
+	}
+		
+
+	return UAnimGenControls::SpecifyStructControlFromArrayViews(InControlSchema,
+		{
+			TEXT("Trajectory"),
+			TEXT("Tags")
+		},
+		{
+			UAnimGenControls::SpecifyTrajectoryControl(InControlSchema, TrajectorySampleNum),
+			UAnimGenControls::SpecifyNamedInclusiveDiscreteControlFromArrayView(InControlSchema, TagNames),
+		});
+}
+
+void UAnimGenBehavior_TrajectoryFollow::MakeControlSets_Implementation(TArray<FAnimGenControlSet>& OutControlSets, FAnimGenControlObject& InControlObject, const UAnimDatabase* InDatabase, const FAnimDatabaseFrameRanges& InFrameRanges) const
+{
+	// Trim TrajectoryFutureTime from all anim sequences since we don't have the future trajectory
+	const FAnimDatabaseFrameRanges TrimmedRanges = UAnimDatabaseFrameRangesLibrary::FrameRangesIntersection(
+		InFrameRanges, 
+		UAnimDatabaseFrameRangesLibrary::FrameRangesTrimEnd(
+			UAnimDatabaseFrameRangesLibrary::FrameRangesAll(InDatabase),
+			FMath::RoundToInt(TrajectoryFutureTime * InDatabase->GetFrameRate().AsDecimal())));
+
+	const int32 TagNum = TagRanges.Num();
+
+	TArray<FName, TInlineAllocator<32>> TagNames;
+	TArray<FAnimDatabaseFrameRanges, TInlineAllocator<32>> TagFrameRanges;
+	TagNames.Reserve(TagNum);
+	TagFrameRanges.Reserve(TagNum);
+
+	for (int32 TagIdx = 0; TagIdx < TagNum; TagIdx++)
+	{
+		if (TagNames.Contains(TagRanges[TagIdx].Name)) { continue; }
+
+		const FAnimDatabaseFrameRanges Ranges = UAnimDatabaseFrameRangesLibrary::MakeFrameRangesFromFunction(InDatabase, TrimmedRanges, TagRanges[TagIdx].FrameRanges);
+
+		if (Ranges.IsValid() && !Ranges.FrameRangeSet->IsEmpty())
+		{
+			TagNames.Add(TagRanges[TagIdx].Name);
+			TagFrameRanges.Add(Ranges);
+		}
+	}
+
+	OutControlSets.Add(UAnimGenControls::MakeStructControlSetFromArrayViews(InControlObject,
+		{
+			TEXT("Trajectory"),
+			TEXT("Tags")
+		},
+		{
+			UAnimGenControls::MakeTrajectoryControlSetFromDatabase(InControlObject, InDatabase, TrimmedRanges, TrajectorySampleNum, TrajectoryFutureTime, TrajectoryForwardVector),
+			UAnimGenControls::MakeNamedInclusiveDiscreteControlSetFromFrameRangesArrayViews(InControlObject, TrimmedRanges, TagNames, TagFrameRanges),
+		}));
+}
+
+void UAnimGenBehavior_TrajectoryFollow::DrawDebugControl_Implementation(
+	const FDebugDrawer& Drawer,
+	const FAnimGenControlObject& InControlObject,
+	const FAnimGenControlObjectElement InControlObjectElement,
+	const FAnimDatabasePoseState& InPoseState,
+	const UAnimDatabase* InDatabase,
+	const UAnimGenAutoEncoder* InAutoEncoder,
+	const int32 SequenceIdx,
+	const float SequenceTime,
+	const int32 RangeStart,
+	const int32 RangeLength,
+	const FLinearColor& IdentifierColor) const
+{
+	FAnimGenControlObjectElement TrajectoryControl, TagsControl;
+	UAnimGenControls::GetStructControlElement(TrajectoryControl, InControlObject, InControlObjectElement, TEXT("Trajectory"));
+	UAnimGenControls::GetStructControlElement(TagsControl, InControlObject, InControlObjectElement, TEXT("Tags"));
+
+	const FTransform RootTransform = UAnimDatabasePoseStateLibrary::PoseStateRootTransform(InPoseState);
+
+	UAnimGenControls::DrawDebugTrajectoryControl(Drawer, InControlObject, TrajectoryControl, RootTransform, IdentifierColor, 25.0f, 5.0f, 1.0f, 1.0f);
+	UAnimGenControls::DrawDebugNamedInclusiveDiscreteControl(Drawer, InControlObject, TagsControl, TEXT(""), RootTransform.GetLocation() + FVector(0.0, 0.0, 220.0f), FRotator(0.0f, -90.0f, 0.0f), 10.0f, IdentifierColor);
+}
+
+
+FAnimGenControlSchemaElement UAnimGenBehavior_VelocityFollow::SpecifyControl_Implementation(FAnimGenControlSchema& InControlSchema) const
+{
+	const int32 TagNum = TagRanges.Num();
+
+	TArray<FName, TInlineAllocator<32>> TagNames;
+	TagNames.Reserve(TagNum);
+
+	for (int32 TagIdx = 0; TagIdx < TagNum; TagIdx++)
+	{
+		TagNames.AddUnique(TagRanges[TagIdx].Name);
+	}
+
+	return UAnimGenControls::SpecifyStructControlFromArrayViews(InControlSchema,
+		{
+			TEXT("LinearVelocity"),
+			TEXT("AngularVelocity"),
+			TEXT("Tags")
+		},
+		{
+			UAnimGenControls::SpecifyLinearVelocityControl(InControlSchema),
+			UAnimGenControls::SpecifyAngularVelocityControl(InControlSchema),
+			UAnimGenControls::SpecifyNamedInclusiveDiscreteControlFromArrayView(InControlSchema, TagNames),
+		});
+}
+
+void UAnimGenBehavior_VelocityFollow::MakeControlSets_Implementation(TArray<FAnimGenControlSet>& OutControlSets, FAnimGenControlObject& InControlObject, const UAnimDatabase* InDatabase, const FAnimDatabaseFrameRanges& InFrameRanges) const
+{
+	const int32 TagNum = TagRanges.Num();
+
+	TArray<FName, TInlineAllocator<32>> TagNames;
+	TArray<FAnimDatabaseFrameRanges, TInlineAllocator<32>> TagFrameRanges;
+	TagNames.Reserve(TagNum);
+	TagFrameRanges.Reserve(TagNum);
+
+	for (int32 TagIdx = 0; TagIdx < TagNum; TagIdx++)
+	{
+		if (TagNames.Contains(TagRanges[TagIdx].Name)) { continue; }
+
+		const FAnimDatabaseFrameRanges Ranges = UAnimDatabaseFrameRangesLibrary::MakeFrameRangesFromFunction(InDatabase, InFrameRanges, TagRanges[TagIdx].FrameRanges);
+
+		if (Ranges.IsValid() && !Ranges.FrameRangeSet->IsEmpty())
+		{
+			TagNames.Add(TagRanges[TagIdx].Name);
+			TagFrameRanges.Add(Ranges);
+		}
+	}
+
+	OutControlSets.Add(UAnimGenControls::MakeStructControlSetFromArrayViews(InControlObject,
+		{
+			TEXT("LinearVelocity"),
+			TEXT("AngularVelocity"),
+			TEXT("Tags")
+		},
+		{
+			UAnimGenControls::MakeRootLinearVelocityControlSetFromDatabase(InControlObject, InDatabase, InFrameRanges),
+			UAnimGenControls::MakeRootAngularVelocityControlSetFromDatabase(InControlObject, InDatabase, InFrameRanges),
+			UAnimGenControls::MakeNamedInclusiveDiscreteControlSetFromFrameRangesArrayViews(InControlObject, InFrameRanges, TagNames, TagFrameRanges),
+		}));
+}
+
+void UAnimGenBehavior_VelocityFollow::DrawDebugControl_Implementation(
+	const FDebugDrawer& Drawer,
+	const FAnimGenControlObject& InControlObject,
+	const FAnimGenControlObjectElement InControlObjectElement,
+	const FAnimDatabasePoseState& InPoseState,
+	const UAnimDatabase* InDatabase,
+	const UAnimGenAutoEncoder* InAutoEncoder,
+	const int32 SequenceIdx,
+	const float SequenceTime,
+	const int32 RangeStart,
+	const int32 RangeLength,
+	const FLinearColor& IdentifierColor) const
+{
+	FAnimGenControlObjectElement LinearVelocityControl, AngularVelocityControl, TagsControl;
+	UAnimGenControls::GetStructControlElement(LinearVelocityControl, InControlObject, InControlObjectElement, TEXT("LinearVelocity"));
+	UAnimGenControls::GetStructControlElement(AngularVelocityControl, InControlObject, InControlObjectElement, TEXT("AngularVelocity"));
+	UAnimGenControls::GetStructControlElement(TagsControl, InControlObject, InControlObjectElement, TEXT("Tags"));
+
+	const FTransform RootTransform = UAnimDatabasePoseStateLibrary::PoseStateRootTransform(InPoseState);
+
+	UAnimGenControls::DrawDebugLinearVelocityControl(Drawer, InControlObject, LinearVelocityControl, RootTransform, RootTransform.GetLocation() + FVector(0.0f, 0.0f, 10.0f), IdentifierColor, 1.0f, 1.0f);
+	UAnimGenControls::DrawDebugAngularVelocityControl(Drawer, InControlObject, AngularVelocityControl, RootTransform, RootTransform.GetLocation() + FVector(0.0f, 0.0f, 10.0f), IdentifierColor, 1.0f, 1.0f);
+	UAnimGenControls::DrawDebugNamedInclusiveDiscreteControl(Drawer, InControlObject, TagsControl, TEXT(""), RootTransform.GetLocation() + FVector(0.0, 0.0, 220.0f), FRotator(0.0f, -90.0f, 0.0f), 10.0f, IdentifierColor);
+}
+
+FAnimGenControlSchemaElement UAnimGenBehavior_MoveToTarget::SpecifyControl_Implementation(FAnimGenControlSchema& InControlSchema) const
+{
+	const int32 TagNum = TagRanges.Num();
+
+	TArray<FName, TInlineAllocator<32>> TagNames;
+	TagNames.Reserve(TagNum);
+
+	for (int32 TagIdx = 0; TagIdx < TagNum; TagIdx++)
+	{
+		TagNames.AddUnique(TagRanges[TagIdx].Name);
+	}
+
+	return UAnimGenControls::SpecifyStructControlFromArrayViews(InControlSchema,
+		{
+			TEXT("Location"),
+			TEXT("Direction"),
+			TEXT("Tags")
+		},
+		{
+			UAnimGenControls::SpecifyLocationControl(InControlSchema),
+			UAnimGenControls::SpecifyOptionalControl(InControlSchema, UAnimGenControls::SpecifyDirectionControl(InControlSchema)),
+			UAnimGenControls::SpecifyNamedInclusiveDiscreteControlFromArrayView(InControlSchema, TagNames),
+		});
+}
+
+void UAnimGenBehavior_MoveToTarget::MakeControlSets_Implementation(TArray<FAnimGenControlSet>& OutControlSets, FAnimGenControlObject& InControlObject, const UAnimDatabase* InDatabase, const FAnimDatabaseFrameRanges& InFrameRanges) const
+{
+	const int32 TagNum = TagRanges.Num();
+
+	TArray<FName, TInlineAllocator<32>> TagNames;
+	TArray<FAnimDatabaseFrameRanges, TInlineAllocator<32>> TagFrameRanges;
+	TagNames.Reserve(TagNum);
+	TagFrameRanges.Reserve(TagNum);
+
+	for (int32 TagIdx = 0; TagIdx < TagNum; TagIdx++)
+	{
+		if (TagNames.Contains(TagRanges[TagIdx].Name)) { continue; }
+
+		const FAnimDatabaseFrameRanges Ranges = UAnimDatabaseFrameRangesLibrary::MakeFrameRangesFromFunction(InDatabase, InFrameRanges, TagRanges[TagIdx].FrameRanges);
+
+		if (Ranges.IsValid() && !Ranges.FrameRangeSet->IsEmpty())
+		{
+			TagNames.Add(TagRanges[TagIdx].Name);
+			TagFrameRanges.Add(Ranges);
+		}
+	}
+
+	const FAnimGenControlSet LocationControlSet = UAnimGenControls::MakeRootLocationAtRangeEndControlSetFromDatabase(InControlObject, InDatabase, InFrameRanges);
+	const FAnimGenControlSet DirectionControlSet = UAnimGenControls::MakeRootDirectionAtRangeEndControlSetFromDatabase(InControlObject, InDatabase, InFrameRanges);
+	const FAnimGenControlSet TagsControlSet = UAnimGenControls::MakeNamedInclusiveDiscreteControlSetFromFrameRangesArrayViews(InControlObject, InFrameRanges, TagNames, TagFrameRanges);
+
+	OutControlSets.Add(UAnimGenControls::MakeStructControlSetFromArrayViews(InControlObject,
+		{
+			TEXT("Location"),
+			TEXT("Direction"),
+			TEXT("Tags")
+		},
+		{
+			LocationControlSet,
+			UAnimGenControls::MakeOptionalValidControlSet(InControlObject, DirectionControlSet),
+			TagsControlSet,
+		}));
+
+	OutControlSets.Add(UAnimGenControls::MakeStructControlSetFromArrayViews(InControlObject,
+		{
+			TEXT("Location"),
+			TEXT("Direction"),
+			TEXT("Tags")
+		},
+		{
+			LocationControlSet,
+			UAnimGenControls::MakeOptionalNullControlSet(InControlObject, InFrameRanges),
+			TagsControlSet,
+		}));
+}
+
+void UAnimGenBehavior_MoveToTarget::DrawDebugControl_Implementation(
+	const FDebugDrawer& Drawer,
+	const FAnimGenControlObject& InControlObject,
+	const FAnimGenControlObjectElement InControlObjectElement,
+	const FAnimDatabasePoseState& InPoseState,
+	const UAnimDatabase* InDatabase,
+	const UAnimGenAutoEncoder* InAutoEncoder,
+	const int32 SequenceIdx,
+	const float SequenceTime,
+	const int32 RangeStart,
+	const int32 RangeLength,
+	const FLinearColor& IdentifierColor) const
+{
+	FAnimGenControlObjectElement LocationControl, OptionalDirectionControl, TagsControl;
+	UAnimGenControls::GetStructControlElement(LocationControl, InControlObject, InControlObjectElement, TEXT("Location"));
+	UAnimGenControls::GetStructControlElement(OptionalDirectionControl, InControlObject, InControlObjectElement, TEXT("Direction"));
+	UAnimGenControls::GetStructControlElement(TagsControl, InControlObject, InControlObjectElement, TEXT("Tags"));
+
+	const FTransform RootTransform = UAnimDatabasePoseStateLibrary::PoseStateRootTransform(InPoseState);
+
+	UAnimGenControls::DrawDebugLocationControl(Drawer, InControlObject, LocationControl, RootTransform, IdentifierColor, 10.0f, 1.0f);
+
+	FAnimGenControlObjectElement DirectionControl;
+	EAnimGenOptionalControl OptionalControl;
+	UAnimGenControls::GetOptionalControl(OptionalControl, DirectionControl, InControlObject, OptionalDirectionControl);
+	if (OptionalControl == EAnimGenOptionalControl::Valid)
+	{
+		FVector Location = FVector::ZeroVector;
+		UAnimGenControls::GetLocationControl(Location, InControlObject, LocationControl, RootTransform);
+		UAnimGenControls::DrawDebugDirectionControl(Drawer, InControlObject, DirectionControl, RootTransform, Location + FVector(0.0f, 0.0, 1.0f), IdentifierColor, 25.0f, 1.0f, 1.0f);
+	}
+
+	UAnimGenControls::DrawDebugNamedInclusiveDiscreteControl(Drawer, InControlObject, TagsControl, TEXT(""), RootTransform.GetLocation() + FVector(0.0, 0.0, 220.0f), FRotator(0.0f, -90.0f, 0.0f), 10.0f, IdentifierColor);
+}
+
+FAnimGenControlSchemaElement UAnimGenBehavior_TrajectoryInteraction::SpecifyControl_Implementation(FAnimGenControlSchema& InControlSchema) const
+{
+	const int32 TagNum = TagRanges.Num();
+
+	TArray<FName, TInlineAllocator<32>> TagNames;
+	TagNames.Reserve(TagNum);
+
+	for (int32 TagIdx = 0; TagIdx < TagNum; TagIdx++)
+	{
+		TagNames.AddUnique(TagRanges[TagIdx].Name);
+	}
+
+
+	return UAnimGenControls::SpecifyStructControlFromArrayViews(InControlSchema,
+		{
+			TEXT("Trajectory"),
+			TEXT("Tags"),
+			TEXT("Interaction")
+		},
+		{
+			UAnimGenControls::SpecifyTrajectoryControl(InControlSchema, TrajectorySampleNum),
+			UAnimGenControls::SpecifyNamedInclusiveDiscreteControlFromArrayView(InControlSchema, TagNames),
+			UAnimGenControls::SpecifyOptionalControl(InControlSchema, UAnimGenControls::SpecifyStructControlFromArrayViews(InControlSchema,
+				{
+					TEXT("Location"),
+					TEXT("Direction"),
+				},
+				{
+					UAnimGenControls::SpecifyLocationControl(InControlSchema),
+					UAnimGenControls::SpecifyDirectionControl(InControlSchema)
+				}))
+		});
+}
+
+void UAnimGenBehavior_TrajectoryInteraction::MakeControlSets_Implementation(TArray<FAnimGenControlSet>& OutControlSets, FAnimGenControlObject& InControlObject, const UAnimDatabase* InDatabase, const FAnimDatabaseFrameRanges& InFrameRanges) const
+{
+	// Trim TrajectoryFutureTime from all anim sequences since we don't have the future trajectory
+	const FAnimDatabaseFrameRanges TrimmedRanges = UAnimDatabaseFrameRangesLibrary::FrameRangesIntersection(
+		InFrameRanges,
+		UAnimDatabaseFrameRangesLibrary::FrameRangesTrimEnd(
+			UAnimDatabaseFrameRangesLibrary::FrameRangesAll(InDatabase),
+			FMath::RoundToInt(TrajectoryFutureTime * InDatabase->GetFrameRate().AsDecimal())));
+
+	const int32 TagNum = TagRanges.Num();
+
+	TArray<FName, TInlineAllocator<32>> TagNames;
+	TArray<FAnimDatabaseFrameRanges, TInlineAllocator<32>> TagFrameRanges;
+	TagNames.Reserve(TagNum);
+	TagFrameRanges.Reserve(TagNum);
+
+	for (int32 TagIdx = 0; TagIdx < TagNum; TagIdx++)
+	{
+		if (TagNames.Contains(TagRanges[TagIdx].Name)) { continue; }
+
+		const FAnimDatabaseFrameRanges Ranges = UAnimDatabaseFrameRangesLibrary::MakeFrameRangesFromFunction(InDatabase, TrimmedRanges, TagRanges[TagIdx].FrameRanges);
+
+		if (Ranges.IsValid() && !Ranges.FrameRangeSet->IsEmpty())
+		{
+			TagNames.Add(TagRanges[TagIdx].Name);
+			TagFrameRanges.Add(Ranges);
+		}
+	}
+
+	FAnimGenControlSet InteractionFacingControlSet;
+
+	if (bUseMirroredInteractionForwardVector)
+	{
+		const int32 InteractionBoneIndex = InDatabase->FindBoneIndex(InteractionBoneName);
+
+		const FAnimDatabaseFrameAttribute GlobalDirectionAttribute = UAnimDatabaseFrameAttributeLibrary::FrameAttributeSelect(
+			UAnimDatabaseFrameAttributeLibrary::MakeBoolFrameAttributeFromNotMirrored(InDatabase, TrimmedRanges),
+			UAnimDatabaseFrameAttributeLibrary::MakeBoneGlobalDirectionFrameAttribute(InDatabase, TrimmedRanges, InteractionBoneIndex, InteractionForwardVector),
+			UAnimDatabaseFrameAttributeLibrary::MakeBoneGlobalDirectionFrameAttribute(InDatabase, TrimmedRanges, InteractionBoneIndex, InteractionMirroredForwardVector));
+
+		InteractionFacingControlSet = UAnimGenControls::MakeDirectionControlSet(InControlObject, GlobalDirectionAttribute, UAnimDatabaseFrameAttributeLibrary::MakeRootTransformFrameAttribute(InDatabase, TrimmedRanges));
+	}
+	else
+	{
+		InteractionFacingControlSet = UAnimGenControls::MakeBoneGlobalDirectionControlSetFromName(InControlObject, InDatabase, TrimmedRanges, InteractionBoneName, InteractionForwardVector);
+	}
+
+	OutControlSets.Add(UAnimGenControls::MakeStructControlSetFromArrayViews(InControlObject,
+		{
+			TEXT("Trajectory"),
+			TEXT("Tags"),
+			TEXT("Interaction")
+		},
+		{
+			UAnimGenControls::MakeTrajectoryControlSetFromDatabase(InControlObject, InDatabase, TrimmedRanges, TrajectorySampleNum, TrajectoryFutureTime, TrajectoryForwardVector),
+			UAnimGenControls::MakeNamedInclusiveDiscreteControlSetFromFrameRangesArrayViews(InControlObject, TrimmedRanges, TagNames, TagFrameRanges),
+			UAnimGenControls::MakeOptionalControlSet(InControlObject, UAnimGenControls::MakeStructControlSetFromArrayViews(InControlObject, 
+					{
+						TEXT("Location"),
+						TEXT("Direction")
+					}, 
+					{
+						UAnimGenControls::MakeBoneGlobalLocationControlSetFromName(InControlObject, InDatabase, TrimmedRanges, InteractionBoneName),
+						InteractionFacingControlSet,
+					}), 
+					UAnimDatabaseFrameAttributeLibrary::MakeBoolFrameAttributeFromActiveRanges(
+						InteractionRanges ? UAnimDatabaseFrameRangesLibrary::MakeFrameRangesFromFunction(InDatabase, TrimmedRanges, InteractionRanges) : TrimmedRanges, 
+						TrimmedRanges))
+		}));
+}
+
+void UAnimGenBehavior_TrajectoryInteraction::DrawDebugControl_Implementation(
+	const FDebugDrawer& Drawer,
+	const FAnimGenControlObject& InControlObject,
+	const FAnimGenControlObjectElement InControlObjectElement,
+	const FAnimDatabasePoseState& InPoseState,
+	const UAnimDatabase* InDatabase,
+	const UAnimGenAutoEncoder* InAutoEncoder,
+	const int32 SequenceIdx,
+	const float SequenceTime,
+	const int32 RangeStart,
+	const int32 RangeLength,
+	const FLinearColor& IdentifierColor) const
+{
+	FAnimGenControlObjectElement TrajectoryControl, TagsControl, InteractionControl;
+	UAnimGenControls::GetStructControlElement(TrajectoryControl, InControlObject, InControlObjectElement, TEXT("Trajectory"));
+	UAnimGenControls::GetStructControlElement(TagsControl, InControlObject, InControlObjectElement, TEXT("Tags"));
+	UAnimGenControls::GetStructControlElement(InteractionControl, InControlObject, InControlObjectElement, TEXT("Interaction"));
+
+	const FTransform RootTransform = UAnimDatabasePoseStateLibrary::PoseStateRootTransform(InPoseState);
+
+	UAnimGenControls::DrawDebugTrajectoryControl(Drawer, InControlObject, TrajectoryControl, RootTransform, IdentifierColor, 25.0f, 5.0f, 1.0f, 1.0f);
+	UAnimGenControls::DrawDebugNamedInclusiveDiscreteControl(Drawer, InControlObject, TagsControl, TEXT(""), RootTransform.GetLocation() + FVector(0.0, 0.0, 220.0f), FRotator(0.0f, -90.0f, 0.0f), 10.0f, IdentifierColor);
+
+	EAnimGenOptionalControl Option;
+	FAnimGenControlObjectElement InteractionControlStruct;
+	UAnimGenControls::GetOptionalControl(Option, InteractionControlStruct, InControlObject, InteractionControl);
+
+	if (Option == EAnimGenOptionalControl::Valid)
+	{
+		FAnimGenControlObjectElement LocationControl, DirectionControl;
+		UAnimGenControls::GetStructControlElement(LocationControl, InControlObject, InteractionControlStruct, TEXT("Location"));
+		UAnimGenControls::GetStructControlElement(DirectionControl, InControlObject, InteractionControlStruct, TEXT("Direction"));
+
+		FVector InteractionLocation = FVector::ZeroVector;
+		UAnimGenControls::GetLocationControl(InteractionLocation, InControlObject, LocationControl, RootTransform);
+
+		UAnimGenControls::DrawDebugLocationControl(Drawer, InControlObject, LocationControl, RootTransform, IdentifierColor, 10.0f, 2.0f);
+		UAnimGenControls::DrawDebugDirectionControl(Drawer, InControlObject, DirectionControl, RootTransform, InteractionLocation, IdentifierColor, 25.0f, 1.0f, 2.0f);
+	}
+}
+
+#endif
+
+FAnimGenControlObjectElement UAnimGenBehaviorLibrary::MakeUncontrolledBehaviorControl(FAnimGenControlObject& InControlObject)
+{
+	return UAnimGenControls::MakeNullControl(InControlObject);
+}
+
+FAnimGenControlObjectElement UAnimGenBehaviorLibrary::MakeIdleBehaviorControl(FAnimGenControlObject& InControlObject, const TArray<FName>& DesiredTags)
+{
+	return MakeIdleBehaviorControlFromArrayView(InControlObject, DesiredTags);
+}
+
+FAnimGenControlObjectElement UAnimGenBehaviorLibrary::MakeIdleBehaviorControlFromArrayView(FAnimGenControlObject& InControlObject, const TArrayView<const FName> DesiredTags)
+{
+	return UAnimGenControls::MakeNamedInclusiveDiscreteControlFromArrayView(InControlObject, DesiredTags);
+}
+
+FAnimGenControlObjectElement UAnimGenBehaviorLibrary::MakeTrajectoryFollowBehaviorControlFromArrayView(
+	FAnimGenControlObject& InControlObject,
+	const FTransformTrajectory& TransformTrajectory,
+	const FTransform& RootBoneTransform,
+	const TArrayView<const FName> DesiredTags,
+	const int32 InTrajectorySampleNum,
+	const float InTrajectoryFutureTime,
+	const FVector InTrajectoryForwardVector)
+{
+	return UAnimGenControls::MakeStructControlFromArrayViews(InControlObject,
+		{
+			TEXT("Trajectory"),
+			TEXT("Tags")
+		},
+		{
+			UAnimGenControls::MakeTrajectoryControlFromTransformTrajectory(InControlObject, TransformTrajectory, RootBoneTransform, InTrajectorySampleNum, InTrajectoryFutureTime, InTrajectoryForwardVector),
+			UAnimGenControls::MakeNamedInclusiveDiscreteControlFromArrayView(InControlObject, DesiredTags),
+		});
+}
+
+FAnimGenControlObjectElement UAnimGenBehaviorLibrary::MakeTrajectoryFollowBehaviorControl(
+	FAnimGenControlObject& InControlObject,
+	const FTransformTrajectory& TransformTrajectory,
+	const FTransform& RootBoneTransform,
+	const TArray<FName>& DesiredTags,
+	const int32 InTrajectorySampleNum,
+	const float InTrajectoryFutureTime,
+	const FVector InTrajectoryForwardVector)
+{
+	return MakeTrajectoryFollowBehaviorControlFromArrayView(InControlObject, TransformTrajectory, RootBoneTransform, DesiredTags, InTrajectorySampleNum, InTrajectoryFutureTime, InTrajectoryForwardVector);
+}
+
+FAnimGenControlObjectElement UAnimGenBehaviorLibrary::MakeMoveToTargetBehaviorControl(
+	FAnimGenControlObject& InControlObject,
+	const FVector TargetLocation,
+	const FVector TargetDirection,
+	const bool bUseTargetDirection,
+	const FTransform& RootBoneTransform,
+	const TArray<FName>& DesiredTags)
+{
+	return MakeMoveToTargetBehaviorControlFromArrayView(InControlObject, TargetLocation, TargetDirection, bUseTargetDirection, RootBoneTransform, DesiredTags);
+}
+
+FAnimGenControlObjectElement UAnimGenBehaviorLibrary::MakeMoveToTargetBehaviorControlFromArrayView(
+	FAnimGenControlObject& InControlObject,
+	const FVector TargetLocation,
+	const FVector TargetDirection,
+	const bool bUseTargetDirection,
+	const FTransform& RootBoneTransform,
+	const TArrayView<const FName> DesiredTags)
+{
+	return UAnimGenControls::MakeStructControlFromArrayViews(InControlObject,
+		{
+			TEXT("Location"),
+			TEXT("Direction"),
+			TEXT("Tags")
+		},
+		{
+			UAnimGenControls::MakeLocationControl(InControlObject, TargetLocation, RootBoneTransform),
+			bUseTargetDirection ?
+				UAnimGenControls::MakeOptionalValidControl(InControlObject, UAnimGenControls::MakeDirectionControl(InControlObject, TargetDirection, RootBoneTransform)) :
+				UAnimGenControls::MakeOptionalNullControl(InControlObject),
+			UAnimGenControls::MakeNamedInclusiveDiscreteControlFromArrayView(InControlObject, DesiredTags),
+		});
+}
+
+FAnimGenControlObjectElement UAnimGenBehaviorLibrary::MakeTrajectoryInteractionBehaviorControl(
+	FAnimGenControlObject& InControlObject,
+	const FTransformTrajectory& TransformTrajectory,
+	const FTransform& RootBoneTransform,
+	const TArray<FName>& DesiredTags,
+	const FVector InteractionLocation,
+	const FVector InteractionDirection,
+	const bool bInteractionActive,
+	const int32 InTrajectorySampleNum,
+	const float InTrajectoryFutureTime,
+	const FVector InTrajectoryForwardVector)
+{
+	return MakeTrajectoryInteractionBehaviorControlFromArrayView(InControlObject, TransformTrajectory, RootBoneTransform, DesiredTags, InteractionLocation, InteractionDirection, bInteractionActive, InTrajectorySampleNum, InTrajectoryFutureTime, InTrajectoryForwardVector);
+}
+
+FAnimGenControlObjectElement UAnimGenBehaviorLibrary::MakeTrajectoryInteractionBehaviorControlFromArrayView(
+	FAnimGenControlObject& InControlObject,
+	const FTransformTrajectory& TransformTrajectory,
+	const FTransform& RootBoneTransform,
+	const TArrayView<const FName> DesiredTags,
+	const FVector InteractionLocation,
+	const FVector InteractionDirection,
+	const bool bInteractionActive,
+	const int32 InTrajectorySampleNum,
+	const float InTrajectoryFutureTime,
+	const FVector InTrajectoryForwardVector)
+{
+	return UAnimGenControls::MakeStructControlFromArrayViews(InControlObject,
+		{
+			TEXT("Trajectory"),
+			TEXT("Tags"),
+			TEXT("Interaction")
+		},
+		{
+			UAnimGenControls::MakeTrajectoryControlFromTransformTrajectory(InControlObject, TransformTrajectory, RootBoneTransform, InTrajectorySampleNum, InTrajectoryFutureTime, InTrajectoryForwardVector),
+			UAnimGenControls::MakeNamedInclusiveDiscreteControlFromArrayView(InControlObject, DesiredTags),
+			bInteractionActive ? 
+				UAnimGenControls::MakeOptionalValidControl(InControlObject, UAnimGenControls::MakeStructControlFromArrayViews(InControlObject,
+					{
+						TEXT("Location"),
+						TEXT("Direction"),
+					},
+					{
+						UAnimGenControls::MakeLocationControl(InControlObject, InteractionLocation, RootBoneTransform),
+						UAnimGenControls::MakeDirectionControl(InControlObject, InteractionDirection, RootBoneTransform),
+					})) :
+				UAnimGenControls::MakeOptionalNullControl(InControlObject)
+		});
+}
+
+
+bool UAnimGenBehaviorLibrary::MakeRandomBool(int32& State)
+{
+	return UE::Learning::Random::SampleFloat((uint32&)State) <= 0.5f;
+}
+
+bool UAnimGenBehaviorLibrary::MakeRandomBoolWithProbability(int32& State, const float TrueProbability)
+{
+	return UE::Learning::Random::SampleFloat((uint32&)State) <= TrueProbability;
+}
+
+bool UAnimGenBehaviorLibrary::SwitchOnRandomBool(int32& State)
+{
+	return MakeRandomBool(State);
+}
+
+bool UAnimGenBehaviorLibrary::SwitchOnRandomBoolWithProbability(int32& State, const float TrueProbability)
+{
+	return MakeRandomBoolWithProbability(State, TrueProbability);
+}
+
+float UAnimGenBehaviorLibrary::MakeRandomFloat(int32& State)
+{
+	return UE::Learning::Random::SampleFloat((uint32&)State);
+}
+
+float UAnimGenBehaviorLibrary::MakeRandomFloatInRange(int32& State, const float Min, const float Max)
+{
+	return UE::Learning::Random::SampleFloat((uint32&)State) * (Max - Min) + Min;
+}
+
+float UAnimGenBehaviorLibrary::MakeRandomFloatChoice(int32& State, const TArray<float>& Values)
+{
+	return MakeRandomFloatChoiceFromArrayView(State, Values);
+}
+
+float UAnimGenBehaviorLibrary::MakeRandomFloatChoiceFromArrayView(int32& State, const TArrayView<const float> Values)
+{
+	if (Values.IsEmpty())
+	{
+		return 0.0f;
+	}
+
+	if (Values.Num() == 1)
+	{
+		return Values[0];
+	}
+
+	return Values[MakeRandomIntegerInRangeInclusive(State, 0, Values.Num() - 1)];
+}
+
+int32 UAnimGenBehaviorLibrary::MakeRandomIntegerInRangeInclusive(int32& State, const int32 Min, const int32 Max)
+{
+	const int32 ClampedMax = FMath::Max(Max, Min);
+	const uint32 RandomValue = UE::Learning::Random::SampleInt((uint32&)State);
+	return Min + RandomValue % ((ClampedMax - Min) + 1);
+}
+
+uint8 UAnimGenBehaviorLibrary::MakeRandomEnumValue(int32& State, const UEnum* Enum)
+{
+	if (Enum == nullptr || Enum->NumEnums() < 2) { return 0; }
+
+	const uint32 RandomValue = UE::Learning::Random::SampleInt((uint32&)State);
+	return Enum->GetValueByIndex(RandomValue % (Enum->NumEnums() - 1));
+}
+
+FVector UAnimGenBehaviorLibrary::MakeRandomVectorFromGaussianDistribution(int32& State, const FVector Mean, const FVector Std)
+{
+	return UE::Learning::Random::SampleVectorGaussian((uint32&)State, Mean, Std);
+}
+
+FVector UAnimGenBehaviorLibrary::MakeRandomDirection3D(int32& State, const FVector Mean, const FVector Std, const FVector Default)
+{
+	return UE::Learning::Random::SampleVectorGaussian((uint32&)State, Mean, Std).GetSafeNormal(UE_SMALL_NUMBER, Default);
+}
+
+FVector UAnimGenBehaviorLibrary::MakeRandomDirection2D(int32& State)
+{
+	return UE::Learning::Random::SampleVectorGaussian((uint32&)State, FVector::ZeroVector, FVector(1.0, 1.0, 0.0)).GetSafeNormal(UE_SMALL_NUMBER, FVector::ForwardVector);
+}
+
+FVector UAnimGenBehaviorLibrary::MakeRandomLocationFromGaussian2D(int32& State, const FVector Mean, const float Std)
+{
+	return MakeRandomVectorFromGaussianDistribution(State, Mean * FVector(1.0, 1.0, 0.0), FVector(Std, Std, 0.0));
+}
+
+void UAnimGenBehaviorLibrary::MakeRandomSequenceIndexAndTimeFromFrameRanges(
+	int32& OutSequenceIdx,
+	float& OutSequenceTime,
+	int32& State,
+	const FAnimDatabaseFrameRanges& FrameRanges,
+	const float FrameDeltaTime)
+{
+	int32 EntryIdx = INDEX_NONE;
+	int32 RangeIdx = INDEX_NONE;
+	int32 RangeFrame = INDEX_NONE;
+	UE::Learning::Random::SampleFrameInFrameRangeSet(
+		EntryIdx,
+		RangeIdx,
+		RangeFrame,
+		(uint32&)State,
+		*FrameRanges.FrameRangeSet);
+
+	OutSequenceIdx = FrameRanges.FrameRangeSet->GetEntrySequence(EntryIdx);
+	OutSequenceTime = (FrameRanges.FrameRangeSet->GetEntryRangeStart(EntryIdx, RangeIdx) + RangeFrame) * FrameDeltaTime;
+}
+
+
+FVector UAnimGenBehaviorLibrary::ProjectOntoFloor(const FVector Location, const float FloorHeight)
+{
+	return FVector(Location.X, Location.Y, FloorHeight);
+}
+
+FVector UAnimGenBehaviorLibrary::ClampToMaxDistanceFromVector(
+	const FVector InVector,
+	const FVector FromVector,
+	const float MaxDistance)
+{
+	const FVector Offset = InVector - FromVector;
+	const float SqDist = Offset.SquaredLength();
+
+	if (SqDist > MaxDistance * MaxDistance)
+	{
+		return FromVector + MaxDistance * Offset.GetSafeNormal();
+	}
+	else
+	{
+		return InVector;
+	}
+}
+
+void UAnimGenBehaviorLibrary::ApplyStickDeadzoneAndSensitivity(
+	FVector& OutAxis,
+	FVector& OutDirection,
+	float& OutLength,
+	bool& bOutNotInDeadzone,
+	const FVector Axis,
+	const float Deadzone,
+	const float Sensitivity)
+{
+	const float SafeDeadzone = FMath::Clamp(Deadzone, 0.0f, 1.0f - UE_SMALL_NUMBER);
+
+	const float AxisLength = Axis.Length();
+
+	if (AxisLength < SafeDeadzone)
+	{
+		OutAxis = FVector::ZeroVector;
+		OutDirection = FVector::ForwardVector;
+		OutLength = 0.0f;
+		bOutNotInDeadzone = false;
+	}
+	else
+	{
+		OutDirection = Axis.GetUnsafeNormal();
+		OutLength = FMath::Pow((AxisLength - SafeDeadzone) / (1.0f - SafeDeadzone), FMath::Max(Sensitivity, UE_SMALL_NUMBER));
+		OutAxis = OutLength * OutDirection;
+		bOutNotInDeadzone = true;
+	}
+}
+
+bool UAnimGenBehaviorLibrary::IsStickNotInDeadzone(const FVector Axis, const float Deadzone)
+{
+	return  Axis.Length() >= FMath::Clamp(Deadzone, 0.0f, 1.0f - UE_SMALL_NUMBER);
+}
+
+bool UAnimGenBehaviorLibrary::IsStickInDeadzone(const FVector Axis, const float Deadzone)
+{
+	return  Axis.Length() < FMath::Clamp(Deadzone, 0.0f, 1.0f - UE_SMALL_NUMBER);
+}
+
+float UAnimGenBehaviorLibrary::ApplyTriggerSensitivity(const float Trigger, const float Sensitivity)
+{
+	return  FMath::Pow(Trigger, FMath::Max(Sensitivity, UE_SMALL_NUMBER));
+}
+
+void UAnimGenBehaviorLibrary::UpdateDesiredVelocityFromAxis(
+	FVector& OutDesiredVelocity,
+	const FVector Axis,
+	const float VelocitySpeed,
+	const float Deadzone,
+	const float Sensitivity)
+{
+	FVector OutAxis = FVector::ZeroVector;
+	FVector OutDirection = FVector::ForwardVector;
+	float OutLength = 0.0f;
+	bool bOutNotInDeadZone = false;
+
+	ApplyStickDeadzoneAndSensitivity(
+		OutAxis,
+		OutDirection,
+		OutLength,
+		bOutNotInDeadZone,
+		Axis,
+		Deadzone,
+		Sensitivity);
+
+	// Desired Velocity
+
+	OutDesiredVelocity = bOutNotInDeadZone ? VelocitySpeed * OutAxis : FVector::ZeroVector;
+}
+
+FRotator UAnimGenBehaviorLibrary::DesiredDirectionToRotation(const FVector Direction)
+{
+	return FVector(Direction.Y, -Direction.X, 0.0f).ToOrientationRotator();
+}
+
+FQuat UAnimGenBehaviorLibrary::DesiredDirectionToRotationQuat(const FVector Direction)
+{
+	return FVector(Direction.Y, -Direction.X, 0.0f).ToOrientationQuat();
+}
+
+void UAnimGenBehaviorLibrary::UpdateDesiredRotationFromAxis(
+	FRotator& OutDesiredRotation,
+	const FVector Axis,
+	const float Deadzone,
+	const float Sensitivity)
+{
+	FVector OutAxis = FVector::ZeroVector;
+	FVector OutDirection = FVector::ForwardVector;
+	float OutLength = 0.0f;
+	bool bOutNotInDeadZone = false;
+
+	ApplyStickDeadzoneAndSensitivity(
+		OutAxis,
+		OutDirection,
+		OutLength,
+		bOutNotInDeadZone,
+		Axis,
+		Deadzone,
+		Sensitivity);
+
+	if (bOutNotInDeadZone)
+	{
+		OutDesiredRotation = DesiredDirectionToRotation(OutDirection);
+	}
+}
+
+void UAnimGenBehaviorLibrary::UpdateDesiredVelocityAndRotationFromAxes(
+	FVector& OutDesiredVelocity,
+	FRotator& OutDesiredRotation,
+	const FVector VelocityAxis,
+	const FVector FacingAxis,
+	const float VelocitySpeed,
+	const float Deadzone,
+	const float Sensitivity)
+{
+	// Velocity Stick
+
+	FVector OutVelocityAxis = FVector::ZeroVector;
+	FVector OutVelocityDirection = FVector::ForwardVector;
+	float OutVelocityLength = 0.0f;
+	bool bOutVelocityNotInDeadZone = false;
+
+	ApplyStickDeadzoneAndSensitivity(
+		OutVelocityAxis,
+		OutVelocityDirection,
+		OutVelocityLength,
+		bOutVelocityNotInDeadZone,
+		VelocityAxis,
+		Deadzone,
+		Sensitivity);
+
+	// Facing Stick
+
+	FVector OutFacingAxis = FVector::ZeroVector;
+	FVector OutFacingDirection = FVector::ForwardVector;
+	float OutFacingLength = 0.0f;
+	bool bOutFacingNotInDeadZone = false;
+
+	ApplyStickDeadzoneAndSensitivity(
+		OutFacingAxis,
+		OutFacingDirection,
+		OutFacingLength,
+		bOutFacingNotInDeadZone,
+		FacingAxis,
+		Deadzone,
+		Sensitivity);
+
+	// Desired Velocity
+
+	OutDesiredVelocity = VelocitySpeed * OutVelocityAxis;
+
+	// Desired Rotation
+
+	if (bOutFacingNotInDeadZone)
+	{
+		OutDesiredRotation = DesiredDirectionToRotation(OutFacingDirection);
+	}
+	else if (bOutVelocityNotInDeadZone)
+	{
+		OutDesiredRotation = DesiredDirectionToRotation(OutVelocityDirection);
+	}
+}
+
+void UAnimGenBehaviorLibrary::UpdateDesiredVelocityAndRotationFromAxisAndCameraTrigger(
+	FVector& OutDesiredVelocity,
+	FRotator& OutDesiredRotation,
+	const FRotator CameraYawRotation,
+	const FVector VelocityAxis,
+	const float FacingTrigger,
+	const float MinimumSpeed,
+	const float MaximumSpeed,
+	const float Deadzone,
+	const float Sensitivity)
+{
+	FVector OutVelocityAxis = FVector::ZeroVector;
+	FVector OutVelocityDirection = FVector::ForwardVector;
+	float OutVelocityLength = 0.0f;
+	bool bOutVelocityNotInDeadZone = false;
+
+	ApplyStickDeadzoneAndSensitivity(
+		OutVelocityAxis,
+		OutVelocityDirection,
+		OutVelocityLength,
+		bOutVelocityNotInDeadZone,
+		CameraYawRotation.RotateVector(VelocityAxis),
+		Deadzone,
+		Sensitivity);
+
+	if (FacingTrigger > 0.0f)
+	{
+		if (bOutVelocityNotInDeadZone)
+		{
+			const FQuat CameraRotation = CameraYawRotation.RotateVector(-FVector::RightVector).ToOrientationQuat();
+			const FQuat FacingRotation = DesiredDirectionToRotationQuat(OutVelocityDirection);
+
+			OutDesiredVelocity = FMath::Lerp(MinimumSpeed, MaximumSpeed, OutVelocityLength) * OutVelocityDirection;
+			OutDesiredRotation = FQuat::Slerp(FacingRotation, CameraRotation, FacingTrigger).Rotator();
+		}
+		else
+		{
+			OutDesiredVelocity = FVector::ZeroVector;
+			OutDesiredRotation = CameraYawRotation.RotateVector(-FVector::RightVector).ToOrientationRotator();
+		}
+	}
+	else if (bOutVelocityNotInDeadZone)
+	{
+		OutDesiredVelocity = FMath::Lerp(MinimumSpeed, MaximumSpeed, OutVelocityLength) * OutVelocityDirection;
+		OutDesiredRotation = DesiredDirectionToRotation(OutVelocityDirection);
+	}
+	else
+	{
+		OutDesiredVelocity = FVector::ZeroVector;
+	}
+}
+
+int32 UAnimGenBehaviorLibrary::FindNearestIndexInFloatArray(const TArray<float>& Values, const float Value)
+{
+	if (Values.IsEmpty())
+	{
+		return INDEX_NONE;
+	}
+
+	const int32 Num = Values.Num();
+
+	float BestDist = UE_MAX_FLT;
+	int32 BestIndex = INDEX_NONE;
+	for (int32 Idx = 0; Idx < Num; Idx++)
+	{
+		const float Dist = FMath::Abs(Values[Idx] - Value);
+		if (Dist < BestDist)
+		{
+			BestDist = Dist;
+			BestIndex = Idx;
+		}
+	}
+
+	return BestIndex;
+}
+
+FName UAnimGenBehaviorLibrary::CycleNameInArray(const TArray<FName>& Names, const FName CurrentName, const int32 Steps)
+{
+	return CycleNameInArrayView(Names, CurrentName, Steps);
+}
+
+FName UAnimGenBehaviorLibrary::CycleNameInArrayView(const TArrayView<const FName> Names, const FName CurrentName, const int32 Steps)
+{
+	int32 Index = Names.Find(CurrentName);
+	if (Index != INDEX_NONE)
+	{
+		return Names[FMath::WrapExclusive(Index + Steps, 0, Names.Num())];
+	}
+
+	return NAME_None;
+}
+
+bool UAnimGenBehaviorLibrary::LocationDifferenceBelowThreshold(
+	const FVector CurrentLocation,
+	const FVector TargetLocation,
+	const float LocationThreshold)
+{
+	return FVector::DistSquared(TargetLocation, CurrentLocation) < LocationThreshold * LocationThreshold;
+}
+
+bool UAnimGenBehaviorLibrary::LocationAndDirectionDifferenceBelowThreshold(
+	const FVector CurrentLocation,
+	const FVector CurrentDirection,
+	const FVector TargetLocation,
+	const FVector TargetDirection,
+	const float LocationThreshold,
+	const float DirectionThresoldAngle)
+{
+	const bool bLocationClose = FVector::DistSquared(TargetLocation, CurrentLocation) < LocationThreshold * LocationThreshold;
+	const bool bDirectionClose = TargetDirection.IsZero() || (FMath::Acos(CurrentDirection.GetSafeNormal().Dot(TargetDirection.GetSafeNormal())) < FMath::DegreesToRadians(DirectionThresoldAngle));
+	return bLocationClose && bDirectionClose;
+}
+
+void UAnimGenBehaviorLibrary::MakeTransformTrajectory(FTransformTrajectory& OutTransformTrajectory, const TArray<FVector>& Locations, const TArray<FRotator>& Rotations, const TArray<float>& Times)
+{
+	MakeTransformTrajectoryFromArrayViews(OutTransformTrajectory, Locations, Rotations, Times);
+}
+
+void UAnimGenBehaviorLibrary::MakeTransformTrajectoryFromArrayViews(FTransformTrajectory& OutTransformTrajectory, const TArrayView<const FVector> Locations, const TArrayView<const FRotator> Rotations, const TArrayView<const float> Times)
+{
+	const int32 SampleNum = FMath::Min(FMath::Min(Locations.Num(), Rotations.Num()), Times.Num());
+
+	OutTransformTrajectory.Samples.SetNumUninitialized(SampleNum);
+	for (int32 SampleIdx = 0; SampleIdx < SampleNum; SampleIdx++)
+	{
+		OutTransformTrajectory.Samples[SampleIdx].Position = Locations[SampleIdx];
+		OutTransformTrajectory.Samples[SampleIdx].Facing = Rotations[SampleIdx].Quaternion();
+		OutTransformTrajectory.Samples[SampleIdx].TimeInSeconds = Times[SampleIdx];
+	}
+}
+
+void UAnimGenBehaviorLibrary::MakeTransformTrajectoryFromPath(FTransformTrajectory& OutTransformTrajectory, const TArray<FVector>& Path, const float StartDistance, const float PathPointDistance, const float PathSpeed, const float FutureTime, const float SampleInterval)
+{
+	return MakeTransformTrajectoryFromPathArrayView(OutTransformTrajectory, Path, StartDistance, PathPointDistance, PathSpeed, FutureTime, SampleInterval);
+}
+
+void UAnimGenBehaviorLibrary::MakeTransformTrajectoryFromPathArrayView(FTransformTrajectory& OutTransformTrajectory, const TArrayView<const FVector> Path, const float StartDistance, const float PathPointDistance, const float PathSpeed, const float FutureTime, const float SampleInterval)
+{
+	const int32 SampleNum = FMath::Max(FMath::RoundToInt(FutureTime / FMath::Max(SampleInterval, UE_SMALL_NUMBER)), 0);
+	OutTransformTrajectory.Samples.SetNumUninitialized(SampleNum);
+	for (int32 SampleIdx = 0; SampleIdx < SampleNum; SampleIdx++)
+	{
+		const float SamplePathTime = StartDistance + PathSpeed * SampleIdx * SampleInterval;
+
+		OutTransformTrajectory.Samples[SampleIdx].Position = PathLocationAtTimeFromArrayView(Path, SamplePathTime / FMath::Max(PathPointDistance, UE_SMALL_NUMBER));
+		OutTransformTrajectory.Samples[SampleIdx].Facing = PathRotationAtTimeFromArrayView(Path, SamplePathTime / FMath::Max(PathPointDistance, UE_SMALL_NUMBER)).Quaternion();
+		OutTransformTrajectory.Samples[SampleIdx].TimeInSeconds = SampleIdx * SampleInterval;
+	}
+}
+
+void UAnimGenBehaviorLibrary::MakeTransformTrajectoryFromTrajectory(FTransformTrajectory& OutTransformTrajectory, const FAnimGenTrajectory& Trajectory, const float FutureTime)
+{
+	const int32 SampleNum = FMath::Min(Trajectory.Locations.Num(), Trajectory.Rotations.Num());
+
+	OutTransformTrajectory.Samples.SetNumUninitialized(SampleNum);
+	for (int32 SampleIdx = 0; SampleIdx < SampleNum; SampleIdx++)
+	{
+		OutTransformTrajectory.Samples[SampleIdx].Position = Trajectory.Locations[SampleIdx];
+		OutTransformTrajectory.Samples[SampleIdx].Facing = Trajectory.Rotations[SampleIdx].Quaternion();
+		OutTransformTrajectory.Samples[SampleIdx].TimeInSeconds = FutureTime * (((float)SampleIdx) / FMath::Max(SampleNum - 1, 1));
+	}
+}
+
+void UAnimGenBehaviorLibrary::ProjectTransformTrajectoryOntoFloor(FTransformTrajectory& InOutTransformTrajectory, const float FloorHeight)
+{
+	const int32 SampleNum = InOutTransformTrajectory.Samples.Num();
+	for (int32 SampleIdx = 0; SampleIdx < SampleNum; SampleIdx++)
+	{
+		InOutTransformTrajectory.Samples[SampleIdx].Position.Z = FloorHeight;
+	}
+}
+
+int32 UAnimGenBehaviorLibrary::GetTrajectorySampleNum(const FAnimGenTrajectory& Trajectory)
+{
+	return Trajectory.Locations.Num();
+}
+
+const TArray<FVector>& UAnimGenBehaviorLibrary::GetTrajectoryLocations(const FAnimGenTrajectory& Trajectory)
+{
+	return Trajectory.Locations;
+}
+
+const TArray<FVector>& UAnimGenBehaviorLibrary::GetTrajectoryDirections(const FAnimGenTrajectory& Trajectory)
+{
+	return Trajectory.Directions;
+}
+
+void UAnimGenBehaviorLibrary::ResetTrajectory(
+	FAnimGenTrajectory& Trajectory,
+	const int32 SampleNum,
+	const FVector InitialLocation,
+	const FRotator InitialRotation,
+	const FVector ForwardVector)
+{
+	if (SampleNum < 0)
+	{
+		// TODO: Warning
+		return;
+	}
+
+	Trajectory.Locations.Init(InitialLocation, SampleNum);
+	Trajectory.Velocities.Init(FVector::ZeroVector, SampleNum);
+	Trajectory.Accelerations.Init(FVector::ZeroVector, SampleNum);
+	Trajectory.Rotations.Init(InitialRotation, SampleNum);
+	Trajectory.AngularVelocities.Init(FVector::ZeroVector, SampleNum);
+	Trajectory.Directions.Init(InitialRotation.RotateVector(ForwardVector), SampleNum);
+}
+
+void UAnimGenBehaviorLibrary::UpdateAndPredictFutureTrajectory(
+	FAnimGenTrajectory& Trajectory,
+	const FVector DesiredVelocity,
+	const FRotator DesiredRotation,
+	const float DeltaTime,
+	const float FutureTime,
+	const float MaximumDesiredVelocity,
+	const float MaximumAcceleration,
+	const float MaximumAngularVelocity,
+	const FVector ForwardVector)
+{
+	const int32 SampleNum = Trajectory.Locations.Num();
+
+	if (SampleNum < 1)
+	{
+		// TODO: Warning
+		return;
+	}
+
+	const float VelocityHalfLife = UE::AnimDatabase::Math::DampingToHalfLife(
+		2.0f * ((MaximumAcceleration / FMath::Max(MaximumDesiredVelocity, UE_KINDA_SMALL_NUMBER)) * UE_EULERS_NUMBER));
+
+	UE::AnimDatabase::Math::CriticalSpringUpdatePositionFromVelocity(
+		Trajectory.Locations[0],
+		Trajectory.Velocities[0],
+		Trajectory.Accelerations[0],
+		DesiredVelocity,
+		VelocityHalfLife,
+		DeltaTime);
+
+	for (int32 SampleIdx = 1; SampleIdx < SampleNum; SampleIdx++)
+	{
+		Trajectory.Locations[SampleIdx] = Trajectory.Locations[0];
+		Trajectory.Velocities[SampleIdx] = Trajectory.Velocities[0];
+		Trajectory.Accelerations[SampleIdx] = Trajectory.Accelerations[0];
+
+		UE::AnimDatabase::Math::CriticalSpringUpdatePositionFromVelocity(
+			Trajectory.Locations[SampleIdx],
+			Trajectory.Velocities[SampleIdx],
+			Trajectory.Accelerations[SampleIdx],
+			DesiredVelocity,
+			VelocityHalfLife,
+			(FutureTime / (SampleNum - 1)) * SampleIdx);
+	}
+
+	const float RotationHalfLife = UE::AnimDatabase::Math::DampingToHalfLife(
+		2.0f * ((MaximumAngularVelocity / 180.0f) * UE_EULERS_NUMBER));
+
+	UE::AnimDatabase::Math::CriticalSpringUpdate(
+		Trajectory.Rotations[0],
+		Trajectory.AngularVelocities[0],
+		DesiredRotation,
+		RotationHalfLife,
+		DeltaTime);
+
+	for (int32 SampleIdx = 1; SampleIdx < SampleNum; SampleIdx++)
+	{
+		Trajectory.Rotations[SampleIdx] = Trajectory.Rotations[0];
+		Trajectory.AngularVelocities[SampleIdx] = Trajectory.AngularVelocities[0];
+
+		UE::AnimDatabase::Math::CriticalSpringUpdate(
+			Trajectory.Rotations[SampleIdx],
+			Trajectory.AngularVelocities[SampleIdx],
+			DesiredRotation,
+			RotationHalfLife,
+			(FutureTime / (SampleNum - 1)) * SampleIdx);
+	}
+
+	for (int32 SampleIdx = 0; SampleIdx < SampleNum; SampleIdx++)
+	{
+		Trajectory.Directions[SampleIdx] = Trajectory.Rotations[SampleIdx].RotateVector(ForwardVector);
+	}
+}
+
+
+void UAnimGenBehaviorLibrary::TeleportTrajectory(
+	FAnimGenTrajectory& Trajectory,
+	const FVector NewLocation,
+	const FRotator NewRotation)
+{
+	if (Trajectory.Locations.Num() < 1)
+	{
+		// TODO: Warning
+		return;
+	}
+
+	const FVector SourceLocation = Trajectory.Locations[0];
+	const FQuat SourceRotation = Trajectory.Rotations[0].Quaternion();
+	const FVector DestinationLocation = NewLocation;
+	const FQuat DestinationRotation = NewRotation.Quaternion();
+
+	const int32 SampleNum = Trajectory.Locations.Num();
+
+	for (int32 SampleIdx = 0; SampleIdx < SampleNum; SampleIdx++)
+	{
+		Trajectory.Locations[SampleIdx] = DestinationRotation.RotateVector(SourceRotation.UnrotateVector(Trajectory.Locations[SampleIdx] - SourceLocation)) + DestinationLocation;
+		Trajectory.Velocities[SampleIdx] = DestinationRotation.RotateVector(SourceRotation.UnrotateVector(Trajectory.Velocities[SampleIdx]));
+		Trajectory.Accelerations[SampleIdx] = DestinationRotation.RotateVector(SourceRotation.UnrotateVector(Trajectory.Accelerations[SampleIdx]));
+		Trajectory.Rotations[SampleIdx] = (DestinationRotation * SourceRotation.Inverse() * Trajectory.Rotations[SampleIdx].Quaternion()).Rotator();
+		Trajectory.AngularVelocities[SampleIdx] = DestinationRotation.RotateVector(SourceRotation.UnrotateVector(Trajectory.AngularVelocities[SampleIdx]));
+		Trajectory.Directions[SampleIdx] = DestinationRotation.RotateVector(SourceRotation.UnrotateVector(Trajectory.Directions[SampleIdx]));
+	}
+}
+
+FVector UAnimGenBehaviorLibrary::PathVelocityAtTime(const TArray<FVector>& Points, const float Time)
+{
+	return PathVelocityAtTimeFromArrayView(Points, Time);
+}
+
+FVector UAnimGenBehaviorLibrary::PathVelocityAtTimeFromArrayView(const TArrayView<const FVector> Points, const float Time)
+{
+	if (Points.IsEmpty())
+	{
+		return FVector::ZeroVector;
+	}
+
+	FVector OutLocation = FVector::ZeroVector;
+	FVector3f OutVelocity = FVector3f::ZeroVector;
+	UE::AnimDatabase::Math::CubicMonoSampleLocation(OutLocation, OutVelocity, Points, Time, 1.0f);
+	return ((FVector)OutVelocity);
+}
+
+FVector UAnimGenBehaviorLibrary::PathLocationAtTime(const TArray<FVector>& Points, const float Time)
+{
+	return PathLocationAtTimeFromArrayView(Points, Time);
+}
+
+FVector UAnimGenBehaviorLibrary::PathLocationAtTimeFromArrayView(const TArrayView<const FVector> Points, const float Time)
+{
+	if (Points.IsEmpty())
+	{
+		return FVector::ZeroVector;
+	}
+
+	FVector OutLocation = FVector::ZeroVector;
+	UE::AnimDatabase::Math::CubicMonoSampleLocation(OutLocation, Points, Time, 1.0f);
+	return OutLocation;
+}
+
+FVector UAnimGenBehaviorLibrary::PathDirectionAtTime(const TArray<FVector>& Points, const float Time)
+{
+	return PathDirectionAtTimeFromArrayView(Points, Time);
+}
+
+FVector UAnimGenBehaviorLibrary::PathDirectionAtTimeFromArrayView(const TArrayView<const FVector> Points, const float Time)
+{
+	if (Points.IsEmpty())
+	{
+		return FVector::ForwardVector;
+	}
+
+	FVector OutLocation = FVector::ZeroVector;
+	FVector3f OutVelocity = FVector3f::ZeroVector;
+	UE::AnimDatabase::Math::CubicMonoSampleLocation(OutLocation, OutVelocity, Points, Time, 1.0f);
+	return ((FVector)OutVelocity).GetSafeNormal(UE_SMALL_NUMBER, FVector::ForwardVector);
+}
+
+FRotator UAnimGenBehaviorLibrary::PathRotationAtTime(const TArray<FVector>& Points, const float Time)
+{
+	return PathRotationAtTimeFromArrayView(Points, Time);
+}
+
+FRotator UAnimGenBehaviorLibrary::PathRotationAtTimeFromArrayView(const TArrayView<const FVector> Points, const float Time)
+{
+	if (Points.IsEmpty())
+	{
+		return FRotator::ZeroRotator;
+	}
+
+	FVector OutLocation = FVector::ZeroVector;
+	FVector3f OutVelocity = FVector3f::ZeroVector;
+	UE::AnimDatabase::Math::CubicMonoSampleLocation(OutLocation, OutVelocity, Points, Time, 1.0f);
+	const FVector OutDirection = ((FVector)OutVelocity).GetSafeNormal(UE_SMALL_NUMBER, FVector::ForwardVector);
+	return FVector(OutDirection.Y, -OutDirection.X, 0.0f).ToOrientationRotator();
+}
+
+bool UAnimGenBehaviorLibrary::PathAddPoint(TArray<FVector>& Points, const FVector NewPoint, const float MinDistance)
+{
+	if (Points.IsEmpty())
+	{
+		Points.Add(NewPoint);
+		return true;
+	}
+	else
+	{
+		if ((Points.Last() - NewPoint).Length() > MinDistance)
+		{
+			Points.Add(NewPoint);
+			return true;
+		}
+		else
+		{
+			return false;
+		}
+	}
+}
+
+float UAnimGenBehaviorLibrary::PathTotalDistance(const TArray<FVector>& Points)
+{
+	return PathTotalDistanceFromArrayView(Points);
+}
+
+float UAnimGenBehaviorLibrary::PathTotalDistanceFromArrayView(const TArrayView<const FVector> Points)
+{
+	const int32 PointNum = Points.Num();
+
+	float Distance = 0.0f;
+	for (int32 PointIdx = 0; PointIdx < PointNum - 1; PointIdx++)
+	{
+		Distance += FVector::Distance(Points[PointIdx + 0], Points[PointIdx + 1]);
+	}
+	return Distance;
+}
+
+FVector UAnimGenBehaviorLibrary::ApplySimpleAvoidanceToVelocity(
+	const FVector Velocity,
+	const TArray<AActor*>& Actors,
+	const AActor* SelfActor,
+	const float MinAvoidanceRadius,
+	const float MaxAvoidanceRadius,
+	const float AvoidanceVelocity)
+{
+	return ApplySimpleAvoidanceToVelocityArrayView(Velocity, Actors, SelfActor, MinAvoidanceRadius, MaxAvoidanceRadius, AvoidanceVelocity);
+}
+
+FVector UAnimGenBehaviorLibrary::ApplySimpleAvoidanceToVelocityArrayView(
+	const FVector Velocity,
+	const TArrayView<AActor* const> Actors,
+	const AActor* SelfActor,
+	const float MinAvoidanceRadius,
+	const float MaxAvoidanceRadius,
+	const float AvoidanceVelocity)
+{
+	if (!SelfActor) { return Velocity; }
+
+	const FVector Location = SelfActor->GetActorLocation();
+
+	FVector OutVelocity = Velocity;
+
+	const int32 ActorNum = Actors.Num();
+	for (int32 ActorIdx = 0; ActorIdx < ActorNum; ActorIdx++)
+	{
+		if (Actors[ActorIdx] == SelfActor) { continue; }
+
+		const FVector ActorLocation = Actors[ActorIdx]->GetActorLocation();
+
+		const float Distance = FVector::Distance(Location, ActorLocation);
+		const float AvoidanceForce = FMath::Clamp((Distance - MinAvoidanceRadius) / FMath::Max(MaxAvoidanceRadius - MinAvoidanceRadius, UE_SMALL_NUMBER), 0.0f, 1.0);
+
+		if (AvoidanceForce < 1.0f)
+		{
+			const FVector AvoidanceDirection = (Location - ActorLocation).GetSafeNormal();
+
+			OutVelocity = FMath::Lerp(OutVelocity, AvoidanceDirection * AvoidanceVelocity, 1.0f - AvoidanceForce);
+		}
+	}
+
+	return OutVelocity;
+}
+
+bool UAnimGenBehaviorLibrary::CanLookAtActor(
+	const FVector LookAtDirection,
+	const AActor* Actor,
+	const AActor* SelfActor,
+	const float DistanceThreshold,
+	const float DotAngleThreshold)
+{
+	if (!Actor || !SelfActor) { return false; }
+
+	const FVector ActorLocation = Actor->GetActorLocation();
+	const FVector SelfActorLocation = SelfActor->GetActorLocation();
+
+	return Actor != SelfActor &&
+		FVector::Distance(ActorLocation, SelfActorLocation) < DistanceThreshold &&
+		(ActorLocation - SelfActorLocation).GetSafeNormal().Dot(LookAtDirection.GetSafeNormal()) > DotAngleThreshold;
+}
+
+AActor* UAnimGenBehaviorLibrary::FindLookAtActor(
+	const FVector LookAtDirection,
+	const TArray<AActor*>& Actors,
+	const AActor* SelfActor,
+	const float DistanceThreshold,
+	const float DotAngleThreshold)
+{
+	return FindLookAtActorArrayView(LookAtDirection, Actors, SelfActor, DistanceThreshold, DotAngleThreshold);
+}
+
+AActor* UAnimGenBehaviorLibrary::FindLookAtActorArrayView(
+	const FVector LookAtDirection,
+	const TArray<AActor*>& Actors,
+	const AActor* SelfActor,
+	const float DistanceThreshold,
+	const float DotAngleThreshold)
+{
+	if (!SelfActor) { return nullptr; }
+
+	const int32 ActorNum = Actors.Num();
+	for (int32 ActorIdx = 0; ActorIdx < ActorNum; ActorIdx++)
+	{
+		if (CanLookAtActor(LookAtDirection, Actors[ActorIdx], SelfActor, DistanceThreshold, DotAngleThreshold))
+		{
+			return Actors[ActorIdx];
+		}
+	}
+
+	return nullptr;
+}
+
+#if WITH_EDITOR
+
+FAnimGenControlObjectElement UAnimGenBehaviorPreview_Uncontrolled::MakeControl_Implementation(FAnimGenControlObject& InControlObject, const FAnimDatabasePoseState& InPoseState) const
+{
+	return UAnimGenBehaviorLibrary::MakeUncontrolledBehaviorControl(InControlObject);
+}
+
+FAnimGenControlObjectElement UAnimGenBehaviorPreview_Idle::MakeControl_Implementation(FAnimGenControlObject& InControlObject, const FAnimDatabasePoseState& InPoseState) const
+{
+	TArray<FName, TInlineAllocator<1>> DesiredTags;
+	DesiredTags.Reserve(1);
+	if (DesiredStyle != NAME_None) { DesiredTags.Add(DesiredStyle); }
+
+	return UAnimGenBehaviorLibrary::MakeIdleBehaviorControlFromArrayView(InControlObject, DesiredTags);
+}
+
+void UAnimGenBehaviorPreview_Idle::OnPreviewInputKeyPressed_Implementation(const FKey Key)
+{
+	if (Key == EKeys::Gamepad_DPad_Left)
+	{
+		StyleIndex = FMath::WrapExclusive((StyleIndex - 1), 0, StyleTags.Num() + 1);
+		DesiredStyle = StyleIndex == 0 ? NAME_None : StyleTags[StyleIndex - 1];
+	}
+
+	if (Key == EKeys::Gamepad_DPad_Right)
+	{
+		StyleIndex = FMath::WrapExclusive((StyleIndex + 1), 0, StyleTags.Num() + 1);
+		DesiredStyle = StyleIndex == 0 ? NAME_None : StyleTags[StyleIndex - 1];
+	}
+}
+
+void UAnimGenBehaviorPreview_Idle::OnPreviewDrawDebug_Implementation(
+	const FDebugDrawer& Drawer,
+	const FAnimGenControlObject& InControlObject,
+	const FAnimGenControlObjectElement InControlObjectElement,
+	const FAnimDatabasePoseState& InPoseState,
+	const FLinearColor& IdentifierColor)
+{
+	const FTransform RootTransform = UAnimDatabasePoseStateLibrary::PoseStateRootTransform(InPoseState);
+
+	UAnimGenControls::DrawDebugNamedInclusiveDiscreteControl(Drawer, InControlObject, InControlObjectElement, TEXT(""), RootTransform.GetLocation() + FVector(0.0, 0.0, 220.0f), FRotator(0.0f, -90.0f, 0.0f), 10.0f, IdentifierColor);
+}
+
+FAnimGenControlObjectElement UAnimGenBehaviorPreview_TrajectoryFollow::MakeControl_Implementation(FAnimGenControlObject& InControlObject, const FAnimDatabasePoseState& InPoseState) const
+{
+	TArray<FName, TInlineAllocator<2>> DesiredTags;
+	DesiredTags.Reserve(2);
+	if (DesiredGait != NAME_None) { DesiredTags.Add(DesiredGait); }
+	if (DesiredStyle != NAME_None) { DesiredTags.Add(DesiredStyle); }
+
+	return UAnimGenBehaviorLibrary::MakeTrajectoryFollowBehaviorControlFromArrayView(
+		InControlObject,
+		TransformTrajectory,
+		UAnimDatabasePoseStateLibrary::PoseStateRootTransform(InPoseState),
+		DesiredTags,
+		TrajectorySampleNum,
+		TrajectoryFutureTime,
+		TrajectoryForwardVector);
+}
+
+void UAnimGenBehaviorPreview_TrajectoryFollow::OnPreviewReset_Implementation(const FAnimDatabasePoseState& InPoseState)
+{
+	UAnimGenBehaviorLibrary::ResetTrajectory(Trajectory, TrajectorySampleNum, FVector::ZeroVector, FRotator::ZeroRotator, TrajectoryForwardVector);
+	bFollowPath = false;
+	bDrawPath = false;
+	FollowPath.Reset();
+	DrawPath.Reset();
+}
+
+void UAnimGenBehaviorPreview_TrajectoryFollow::OnPreviewUpdate_Implementation(const FAnimDatabasePoseState& InPoseState, const float DeltaTime)
+{
+	FVector RightStickAxis = FVector::ZeroVector;
+	FVector RightStickDirection = FVector::ZeroVector;
+	float RightStickLength = 0.0f;
+	bool bRightStickInDeadzone = false;
+	UAnimGenBehaviorLibrary::ApplyStickDeadzoneAndSensitivity(
+		RightStickAxis,
+		RightStickDirection,
+		RightStickLength,
+		bRightStickInDeadzone,
+		GetPreviewRightStickAxis());
+
+	SpringMath::CriticalSpringDamper(CameraYawVelocity, CameraYawAcceleration, CameraRotationSpeed * (float)RightStickAxis.Y, CameraInputSmoothingTime, DeltaTime);
+	SpringMath::CriticalSpringDamper(CameraPitchVelocity, CameraPitchAcceleration, CameraRotationSpeed * (float)RightStickAxis.X, CameraInputSmoothingTime, DeltaTime);
+	SpringMath::CriticalSpringDamper(CameraDistanceVelocity, CameraDistanceAcceleration, CameraZoomSpeed * CameraDistanceDelta, CameraInputSmoothingTime, DeltaTime);
+
+	UpdatePreviewCamera(
+		CameraYawVelocity,
+		CameraPitchVelocity,
+		CameraDistanceVelocity,
+		DeltaTime,
+		bCameraInvertPitch);
+
+	if (bFollowPath && UAnimGenBehaviorLibrary::IsStickNotInDeadzone(GetPreviewLeftStickAxis()))
+	{
+		bFollowPath = false;
+	}
+
+	if (bFollowPath)
+	{
+		UAnimGenBehaviorLibrary::MakeTransformTrajectoryFromPathArrayView(
+			TransformTrajectory,
+			FollowPath,
+			PathDistance,
+			PathPointDistance,
+			PathFollowSpeed,
+			TrajectoryFutureTime);
+
+		if (bPathFacingCamera)
+		{
+			const int32 SampleNum = TransformTrajectory.Samples.Num();
+			for (int32 SampleIdx = 0; SampleIdx < SampleNum; SampleIdx++)
+			{
+				TransformTrajectory.Samples[SampleIdx].Facing = GetPreviewCameraYawRotation().RotateVector(-FVector::RightVector).ToOrientationQuat();
+			}
+		}
+
+		PathDistance += DeltaTime * PathFollowSpeed;
+	}
+	else
+	{
+		SpringMath::CriticalSpringDamper(DesiredMinimumSpeed, DesiredMinimumSpeedVelocity, bWalkActive ? MinimumWalkSpeed : MinimumRunSpeed, SpeedSmoothingTime, DeltaTime);
+		SpringMath::CriticalSpringDamper(DesiredMaximumSpeed, DesiredMaximumSpeedVelocity, bWalkActive ? MaximumWalkSpeed : MaximumRunSpeed, SpeedSmoothingTime, DeltaTime);
+
+		UAnimGenBehaviorLibrary::UpdateDesiredVelocityAndRotationFromAxisAndCameraTrigger(
+			DesiredVelocity,
+			DesiredRotation,
+			GetPreviewCameraYawRotation(),
+			GetPreviewLeftStickAxis(),
+			GetPreviewLeftTriggerAxisWithSensitivity(),
+			DesiredMinimumSpeed,
+			DesiredMaximumSpeed);
+
+		UAnimGenBehaviorLibrary::TeleportTrajectory(
+			Trajectory,
+			UAnimDatabasePoseStateLibrary::PoseStateRootLocation(InPoseState),
+			UAnimDatabasePoseStateLibrary::PoseStateRootRotation(InPoseState));
+
+		UAnimGenBehaviorLibrary::UpdateAndPredictFutureTrajectory(
+			Trajectory,
+			DesiredVelocity,
+			DesiredRotation,
+			DeltaTime,
+			TrajectoryFutureTime,
+			DesiredMaximumSpeed,
+			MaximumAcceleration,
+			MaximumAngularVelocity,
+			TrajectoryForwardVector);
+
+		UAnimGenBehaviorLibrary::MakeTransformTrajectoryFromTrajectory(TransformTrajectory, Trajectory, TrajectoryFutureTime);
+
+		if (DesiredVelocity.IsNearlyZero())
+		{
+			DesiredGait = StandTag;
+		}
+		else
+		{
+			DesiredGait = bWalkActive ? WalkTag : RunTag;
+		}
+	}
+}
+
+void UAnimGenBehaviorPreview_TrajectoryFollow::OnPreviewInputKeyPressed_Implementation(const FKey Key)
+{
+	if (Key == EKeys::Gamepad_RightShoulder)
+	{
+		bWalkActive = !bWalkActive;
+	}
+
+	if (Key == EKeys::Gamepad_DPad_Left)
+	{
+		StyleIndex = FMath::WrapExclusive((StyleIndex - 1), 0, StyleTags.Num() + 1);
+		DesiredStyle = StyleIndex == 0 ? NAME_None : StyleTags[StyleIndex - 1];
+	}
+
+	if (Key == EKeys::Gamepad_DPad_Right)
+	{
+		StyleIndex = FMath::WrapExclusive((StyleIndex + 1), 0, StyleTags.Num() + 1);
+		DesiredStyle = StyleIndex == 0 ? NAME_None : StyleTags[StyleIndex - 1];
+	}
+
+	if (Key == EKeys::Gamepad_DPad_Up)
+	{
+		CameraDistanceDelta = +1.0f;
+	}
+
+	if (Key == EKeys::Gamepad_DPad_Down)
+	{
+		CameraDistanceDelta = -1.0f;
+	}
+}
+
+void UAnimGenBehaviorPreview_TrajectoryFollow::OnPreviewInputKeyReleased_Implementation(const FKey Key)
+{
+	if (Key == EKeys::Gamepad_DPad_Up)
+	{
+		CameraDistanceDelta = 0.0f;
+	}
+
+	if (Key == EKeys::Gamepad_DPad_Down)
+	{
+		CameraDistanceDelta = 0.0f;
+	}
+}
+
+void UAnimGenBehaviorPreview_TrajectoryFollow::OnPreviewDrawDebug_Implementation(
+	const FDebugDrawer& Drawer,
+	const FAnimGenControlObject& InControlObject,
+	const FAnimGenControlObjectElement InControlObjectElement,
+	const FAnimDatabasePoseState& InPoseState,
+	const FLinearColor& IdentifierColor)
+{
+	FAnimGenControlObjectElement TrajectoryControl, TagsControl;
+	UAnimGenControls::GetStructControlElement(TrajectoryControl, InControlObject, InControlObjectElement, TEXT("Trajectory"));
+	UAnimGenControls::GetStructControlElement(TagsControl, InControlObject, InControlObjectElement, TEXT("Tags"));
+
+	const FTransform RootTransform = UAnimDatabasePoseStateLibrary::PoseStateRootTransform(InPoseState);
+
+	UAnimGenControls::DrawDebugTrajectoryControl(Drawer, InControlObject, TrajectoryControl, RootTransform, IdentifierColor, 25.0f, 5.0f, 1.0f, 1.0f);
+	UAnimGenControls::DrawDebugNamedInclusiveDiscreteControl(Drawer, InControlObject, TagsControl, TEXT(""), RootTransform.GetLocation() + FVector(0.0, 0.0, 220.0f), FRotator(0.0f, -90.0f, 0.0f), 10.0f, IdentifierColor);
+
+	if (bDrawPath)
+	{
+		UDrawDebugLibrary::DrawDebugCatmullRomSplineArrayView(
+			Drawer,
+			DrawPath,
+			UDrawDebugLibrary::MakeDrawDebugLineStyleFromColorAndThickness(IdentifierColor, 2.0f), true, true, 8);
+
+		UDrawDebugLibrary::DrawDebugLocationsArrayView(
+			Drawer, 
+			DrawPath, 
+			UDrawDebugLibrary::MakeDrawDebugLineStyleFromColorAndThickness(IdentifierColor, 2.0f), true, 5.0f);
+	}
+	else if (bFollowPath)
+	{
+		UDrawDebugLibrary::DrawDebugCatmullRomSplineArrayView(
+			Drawer, 
+			FollowPath, 
+			UDrawDebugLibrary::MakeDrawDebugLineStyleFromColorAndThickness(IdentifierColor, 2.0f), true, true, 8);
+
+		UDrawDebugLibrary::DrawDebugLocationsArrayView(
+			Drawer,
+			FollowPath,
+			UDrawDebugLibrary::MakeDrawDebugLineStyleFromColorAndThickness(IdentifierColor, 2.0f), true, 5.0f);
+	}
+}
+
+void UAnimGenBehaviorPreview_TrajectoryFollow::OnPreviewInputLeftMousePressed_Implementation(const FVector MouseGroundLocation)
+{
+	bDrawPath = true;
+	DrawPath.Reset();
+	DrawPath.Add(MouseGroundLocation);
+}
+
+void UAnimGenBehaviorPreview_TrajectoryFollow::OnPreviewInputLeftMouseReleased_Implementation(const FVector MouseGroundLocation)
+{
+	bFollowPath = true;
+	bDrawPath = false;
+	FollowPath = DrawPath;
+	PathDistance = -PathFollowSpeed;
+}
+
+void UAnimGenBehaviorPreview_TrajectoryFollow::OnPreviewInputMouseMoved_Implementation(const FVector MouseGroundLocation)
+{
+	if (bDrawPath)
+	{
+		UAnimGenBehaviorLibrary::PathAddPoint(DrawPath, MouseGroundLocation, PathPointDistance);
+	}
+}
+
+FAnimGenControlObjectElement UAnimGenBehaviorPreview_MoveToTarget::MakeControl_Implementation(FAnimGenControlObject& InControlObject, const FAnimDatabasePoseState& InPoseState) const
+{
+	TArray<FName, TInlineAllocator<1>> DesiredTags;
+	DesiredTags.Reserve(1);
+	if (DesiredStyle != NAME_None) { DesiredTags.Add(DesiredStyle); }
+
+	return UAnimGenBehaviorLibrary::MakeMoveToTargetBehaviorControlFromArrayView(
+		InControlObject,
+		TargetLocation,
+		TargetDirection,
+		bUseTargetDirection,
+		UAnimDatabasePoseStateLibrary::PoseStateRootTransform(InPoseState),
+		DesiredTags);
+}
+
+void UAnimGenBehaviorPreview_MoveToTarget::OnPreviewInputKeyPressed_Implementation(const FKey Key)
+{
+	if (Key == EKeys::Gamepad_DPad_Left)
+	{
+		StyleIndex = FMath::WrapExclusive((StyleIndex - 1), 0, StyleTags.Num() + 1);
+		DesiredStyle = StyleIndex == 0 ? NAME_None : StyleTags[StyleIndex - 1];
+	}
+
+	if (Key == EKeys::Gamepad_DPad_Right)
+	{
+		StyleIndex = FMath::WrapExclusive((StyleIndex + 1), 0, StyleTags.Num() + 1);
+		DesiredStyle = StyleIndex == 0 ? NAME_None : StyleTags[StyleIndex - 1];
+	}
+}
+
+void UAnimGenBehaviorPreview_MoveToTarget::OnPreviewDrawDebug_Implementation(
+	const FDebugDrawer& Drawer,
+	const FAnimGenControlObject& InControlObject,
+	const FAnimGenControlObjectElement InControlObjectElement,
+	const FAnimDatabasePoseState& InPoseState,
+	const FLinearColor& IdentifierColor)
+{
+	if (bPlacingActive)
+	{
+		UDrawDebugLibrary::DrawDebugLocation(
+			Drawer,
+			ActiveLocation,
+			UDrawDebugLibrary::MakeDrawDebugLineStyleFromColorAndThickness(FLinearColor::White.CopyWithNewOpacity(IdentifierColor.A), 1.0f),
+			true,
+			10.0f);
+
+		if (bUseActiveDirection)
+		{
+			UDrawDebugLibrary::DrawDebugDirection(
+				Drawer,
+				ActiveLocation,
+				ActiveDirection,
+				UDrawDebugLibrary::MakeDrawDebugLineStyleFromColorAndThickness(FLinearColor::White.CopyWithNewOpacity(IdentifierColor.A), 1.0f),
+				true,
+				25.0f,
+				1.0f);
+		}
+	}
+
+	UDrawDebugLibrary::DrawDebugLocation(
+		Drawer,
+		TargetLocation,
+		UDrawDebugLibrary::MakeDrawDebugLineStyleFromColorAndThickness(IdentifierColor, 1.0f),
+		true,
+		10.0f);
+
+	if (bUseTargetDirection)
+	{
+		UDrawDebugLibrary::DrawDebugDirection(
+			Drawer,
+			TargetLocation,
+			TargetDirection,
+			UDrawDebugLibrary::MakeDrawDebugLineStyleFromColorAndThickness(IdentifierColor, 1.0f),
+			true,
+			25.0f,
+			1.0f);
+	}
+}
+
+void UAnimGenBehaviorPreview_MoveToTarget::OnPreviewInputLeftMousePressed_Implementation(const FVector MouseGroundLocation)
+{
+	ActiveLocation = MouseGroundLocation;
+	ActiveDirection = FVector::ForwardVector;
+	bUseActiveDirection = false;
+	bPlacingActive = true;
+}
+
+void UAnimGenBehaviorPreview_MoveToTarget::OnPreviewInputLeftMouseReleased_Implementation(const FVector MouseGroundLocation)
+{
+	TargetLocation = ActiveLocation;
+	TargetDirection = ActiveDirection;
+	bUseTargetDirection = bUseActiveDirection;
+	bPlacingActive = false;
+}
+
+void UAnimGenBehaviorPreview_MoveToTarget::OnPreviewInputMouseMoved_Implementation(const FVector MouseGroundLocation)
+{
+	if (bPlacingActive)
+	{
+		if (FVector::Distance(MouseGroundLocation, ActiveLocation) > 10.0f)
+		{
+			ActiveDirection = (MouseGroundLocation - ActiveLocation).GetSafeNormal();
+			bUseActiveDirection = true;
+		}
+	}
+}
+
+#endif
\ No newline at end of file
diff --git a/Engine/Plugins/Experimental/Animation/AnimGen/Source/AnimGen/Private/AnimGenController.cpp b/Engine/Plugins/Experimental/Animation/AnimGen/Source/AnimGen/Private/AnimGenController.cpp
new file mode 100644
--- /dev/null
+++ b/Engine/Plugins/Experimental/Animation/AnimGen/Source/AnimGen/Private/AnimGenController.cpp
@@ -0,0 +1,356 @@
+// Copyright Epic Games, Inc. All Rights Reserved.
+
+#include "AnimGenController.h"
+
+#include "AnimGenLog.h"
+#include "AnimGenControl.h"
+#include "AnimGenAutoEncoder.h"
+#include "AnimGenBehavior.h"
+
+#include "AnimDatabaseFrameRanges.h"
+
+#include "UObject/Package.h"
+#include "Misc/ScopedSlowTask.h"
+
+#include "Animation/Skeleton.h"
+
+#include "DrawDebugLibrary.h"
+
+#include UE_INLINE_GENERATED_CPP_BY_NAME(AnimGenController)
+
+#define LOCTEXT_NAMESPACE "AnimGenController"
+
+#if WITH_EDITOR
+
+void UAnimGenControllerTrainingSettings::ResetTransient()
+{
+	DatabaseContentHash = 0;
+	ActiveFrameRanges = nullptr;
+	ActiveBehavior = nullptr;
+	ControlSets.Empty();
+	QueryEntries.Empty();
+	SelectedRanges.Empty();
+}
+
+void UAnimGenControllerTrainingSettings::PostEditChangeProperty(struct FPropertyChangedEvent& Event)
+{
+	if (Event.GetMemberPropertyName() == GET_MEMBER_NAME_CHECKED(UAnimGenControllerTrainingSettings, Database))
+	{
+		Database.LoadAsync(FLoadSoftObjectPathAsyncDelegate::CreateLambda([this](const FSoftObjectPath&, UObject*) { ResetTransient(); }));
+	}
+	else if (Event.GetMemberPropertyName() == GET_MEMBER_NAME_CHECKED(UAnimGenControllerTrainingSettings, AutoEncoder))
+	{
+		AutoEncoder.LoadAsync(FLoadSoftObjectPathAsyncDelegate::CreateLambda([this](const FSoftObjectPath&, UObject*) { ResetTransient(); }));
+	}
+}
+
+void UAnimGenControllerTrainingSettings::UpdateQueryEntries()
+{
+	QueryEntries.Reset();
+
+	if (Database && AutoEncoder)
+	{
+		const int32 ControlSetNum = ControlSets.Num();
+
+		for (int32 ControlSetIdx = 0; ControlSetIdx < ControlSetNum; ControlSetIdx++)
+		{
+			const UE::Learning::FFrameRangeSet& FrameRangeSet = ControlSets[ControlSetIdx].Data->FrameRangeSet;
+
+			const int32 EntryNum = FrameRangeSet.GetEntryNum();
+
+			for (int32 EntryIdx = 0; EntryIdx < EntryNum; EntryIdx++)
+			{
+				const int32 EntrySequenceIdx = FrameRangeSet.GetEntrySequence(EntryIdx);
+				const UAnimSequence* AnimSequence = Database->GetAnimSequence(EntrySequenceIdx);
+				const bool bIsMirrored = Database->GetIsMirrored(EntrySequenceIdx);
+
+				const int32 RangeNum = FrameRangeSet.GetEntryRangeNum(EntryIdx);
+
+				for (int32 EntryRangeIdx = 0; EntryRangeIdx < RangeNum; EntryRangeIdx++)
+				{
+					TSharedPtr<UE::AnimDatabase::Editor::FQueryEntry> QueryEntry = MakeShared<UE::AnimDatabase::Editor::FQueryEntry>();
+					QueryEntry->Sequence = AnimSequence;
+					QueryEntry->StartFrame = FrameRangeSet.GetEntryRangeStart(EntryIdx, EntryRangeIdx);
+					QueryEntry->StopFrame = QueryEntry->StartFrame + FrameRangeSet.GetEntryRangeLength(EntryIdx, EntryRangeIdx) - 1;
+					QueryEntry->bIsMirrored = bIsMirrored;
+					QueryEntry->bIsSelected = false;
+
+					// Here we use a hash of the sequence, start, and stop frames to produce the color
+
+					const uint32 HashData[6] =
+					{
+						(uint32)RangeIndentifierColorSeed,
+						(uint32)ControlSetIdx,
+						GetTypeHash(AnimSequence),
+						(uint32)QueryEntry->StartFrame,
+						(uint32)QueryEntry->StopFrame,
+						(uint32)QueryEntry->bIsMirrored
+					};
+
+					const uint32 ColorHash = CityHash32((const char*)HashData, sizeof(HashData));
+					const uint8 Hue = ColorHash & 0x000000FF;
+					const uint8 Saturation = ((ColorHash & 0x0000FF00) >> 8) % 25 + 230;
+					const uint8 Brightness = ((ColorHash & 0x00FF0000) >> 16) % 50 + 205;
+
+					QueryEntry->Color = FLinearColor::MakeFromHSV8(Hue, Saturation, Brightness);
+
+					QueryEntries.Add(QueryEntry);
+				}
+			}
+		}
+	}
+}
+
+bool UAnimGenControllerTrainingSettings::Update()
+{
+	if (const UAnimDatabase* DatabaseObject = Database.Get())
+	{
+		bool bControlSetsUpdateRequired = false;
+
+		// Check if Database has been updated
+
+		const uint32 CurrentContentHash = DatabaseObject->GetContentHash();
+		if (CurrentContentHash != DatabaseContentHash)
+		{
+			DatabaseContentHash = CurrentContentHash;
+			bControlSetsUpdateRequired = true;
+		}
+
+		// Check if FrameRanges Class is out-of-date and the class has been regenerated
+
+		bControlSetsUpdateRequired = bControlSetsUpdateRequired || FrameRanges != ActiveFrameRanges || Behavior != ActiveBehavior;
+
+		if (bControlSetsUpdateRequired)
+		{
+			FScopedSlowTask SlowTask(0.0f, LOCTEXT("RebuildingControlSets", "Rebuilding Control Sets"));
+			SlowTask.MakeDialog();
+
+			// Rebuild Frame Ranges
+
+			if (!FrameRanges)
+			{
+				ActiveFrameRanges = nullptr;
+				QueryRanges = UAnimDatabaseFrameRangesLibrary::MakeFrameRangesFromDatabase(DatabaseObject);
+			}
+			else
+			{
+				ActiveFrameRanges = FrameRanges;
+				if (FrameRanges)
+				{
+					QueryRanges = UAnimDatabaseFrameRangesLibrary::MakeFrameRangesFromFunction(DatabaseObject, UAnimDatabaseFrameRangesLibrary::MakeFrameRangesFromDatabase(DatabaseObject), FrameRanges);
+					if (!QueryRanges.IsValid()) { QueryRanges = UAnimDatabaseFrameRangesLibrary::MakeEmptyFrameRanges(); }
+				}
+				else
+				{
+					QueryRanges = UAnimDatabaseFrameRangesLibrary::MakeFrameRangesFromDatabase(DatabaseObject);
+				}
+
+				if (!QueryRanges.IsValid()) { QueryRanges = UAnimDatabaseFrameRangesLibrary::MakeEmptyFrameRanges(); }
+			}
+
+			// Rebuild Control Sets
+
+			if (Behavior)
+			{
+				ActiveBehavior = Behavior;
+
+				ControlSchema = UAnimGenControls::MakeControlSchema();
+				ControlObject = UAnimGenControls::MakeControlObject();
+				ControlSchemaElement = FAnimGenControlSchemaElement();
+				ControlSets.Empty();
+				ControlSchemaElement = Behavior->SpecifyControl(ControlSchema);
+				if (ControlSchema.ObservationSchema->Schema.IsValid(ControlSchemaElement.SchemaElement))
+				{
+					Behavior->MakeControlSets(ControlSets, ControlObject, DatabaseObject, QueryRanges);
+				}
+				else
+				{
+					ControlSets.Empty();
+				}
+			}
+			else
+			{
+				ActiveBehavior = nullptr;
+				ControlSets.Empty();
+			}
+
+			UpdateQueryEntries();
+			SelectedRanges.Empty();
+		}
+
+		// Update Selection
+
+		SelectedRanges.Reset();
+
+		const int32 QueryRangeNum = QueryEntries.Num();
+
+		for (int32 QueryRangeIdx = 0; QueryRangeIdx < QueryRangeNum; QueryRangeIdx++)
+		{
+			if (QueryEntries[QueryRangeIdx]->bIsSelected)
+			{
+				SelectedRanges.Add(QueryRangeIdx);
+			}
+		}
+
+		return bControlSetsUpdateRequired;
+	}
+	else
+	{
+		if (DatabaseContentHash != 0)
+		{
+			ResetTransient();
+			return true;
+		}
+		else
+		{
+			return false;
+		}
+	}
+}
+
+void UAnimGenControllerTestingSettings::PostEditChangeProperty(struct FPropertyChangedEvent& Event)
+{
+	if (Event.GetMemberPropertyName() == GET_MEMBER_NAME_CHECKED(UAnimGenControllerTestingSettings, ResetDatabase))
+	{
+		ResetDatabase.LoadAsync(FLoadSoftObjectPathAsyncDelegate());
+	}
+}
+
+#endif
+
+UAnimGenController::UAnimGenController(const FObjectInitializer& ObjectInitializer)
+{
+#if WITH_EDITOR
+	TrainingSettings = ObjectInitializer.CreateEditorOnlyDefaultSubobject<UAnimGenControllerTrainingSettings>(this, TEXT("TrainingSettings"));
+	TestingSettings = ObjectInitializer.CreateEditorOnlyDefaultSubobject<UAnimGenControllerTestingSettings>(this, TEXT("TestingSettings"));
+	ViewportSettings = ObjectInitializer.CreateEditorOnlyDefaultSubobject<UAnimGenControllerViewportSettings>(this, TEXT("ViewportSettings"));
+#endif
+}
+
+bool UAnimGenController::IsValid() const
+{
+	return 
+		AutoEncoder &&
+		AutoEncoder->IsValid() &&
+		ControlSchema.IsValid() &&
+		ControlSchema.IsElementValid(ControlSchemaElement) &&
+		ControlEncoderNetwork &&
+		LOD0Network &&
+		LOD1Network &&
+		LOD2Network;
+}
+
+void UAnimGenController::Invalidate()
+{
+	AutoEncoder = nullptr;
+	ControlEncoderNetwork = nullptr;
+	LOD0Network = nullptr;
+	LOD1Network = nullptr;
+	LOD2Network = nullptr;
+	FrameRate = FFrameRate(60, 1);
+	ControlVectorSize = 0;
+	EncodedControlVectorSize = 0;
+	ControlDistributionVectorSize = 0;
+
+#if WITH_EDITOR
+	TrainedFrameRangesContentHash = 0;
+	TrainedAutoEncoderContentHash = 0;
+	TrainedSchemaCompatibilityHash = 0;
+	ControlEncoderSize = 0;
+	DenoiserSize = 0;
+	ControlEncoderInferenceTime = 0;
+	DenoiserInferenceTime = 0;
+#endif
+}
+
+
+void UAnimGenController::NormalizeEncodedPoseVectors(
+	const TLearningArrayView<2, float> OutNormalizedEncodedPoseVectors,
+	const TLearningArrayView<2, const float> InUnnormalizedEncodedPoseVectors) const
+{
+	const int32 FrameNum = OutNormalizedEncodedPoseVectors.Num<0>();
+	const int32 DimNum = OutNormalizedEncodedPoseVectors.Num<1>();
+	check(FrameNum == InUnnormalizedEncodedPoseVectors.Num<0>());
+	check(DimNum == InUnnormalizedEncodedPoseVectors.Num<1>());
+	check(DimNum == EncodedPoseMeans.Num());
+
+	for (int32 FrameIdx = 0; FrameIdx < FrameNum; FrameIdx++)
+	{
+		for (int32 DimIdx = 0; DimIdx < DimNum; DimIdx++)
+		{
+			OutNormalizedEncodedPoseVectors[FrameIdx][DimIdx] = (InUnnormalizedEncodedPoseVectors[FrameIdx][DimIdx] - EncodedPoseMeans[DimIdx]) / FMath::Max(EncodedPoseNormalizationScale, UE_SMALL_NUMBER);
+		}
+	}
+}
+
+void UAnimGenController::NormalizeEncodedPoseVectorsInplace(const TLearningArrayView<2, float> InOutEncodedPoseVectors) const
+{
+	const int32 FrameNum = InOutEncodedPoseVectors.Num<0>();
+	const int32 DimNum = InOutEncodedPoseVectors.Num<1>();
+	check(DimNum == EncodedPoseMeans.Num());
+
+	for (int32 FrameIdx = 0; FrameIdx < FrameNum; FrameIdx++)
+	{
+		for (int32 DimIdx = 0; DimIdx < DimNum; DimIdx++)
+		{
+			InOutEncodedPoseVectors[FrameIdx][DimIdx] = (InOutEncodedPoseVectors[FrameIdx][DimIdx] - EncodedPoseMeans[DimIdx]) / FMath::Max(EncodedPoseNormalizationScale, UE_SMALL_NUMBER);
+		}
+	}
+}
+
+void UAnimGenController::DenormalizeEncodedPoseVectors(
+	const TLearningArrayView<2, float> OutUnnormalizedEncodedPoseVectors,
+	const TLearningArrayView<2, const float> InNormalizedEncodedPoseVectors) const
+{
+	const int32 FrameNum = OutUnnormalizedEncodedPoseVectors.Num<0>();
+	const int32 DimNum = OutUnnormalizedEncodedPoseVectors.Num<1>();
+	check(FrameNum == InNormalizedEncodedPoseVectors.Num<0>());
+	check(DimNum == InNormalizedEncodedPoseVectors.Num<1>());
+	check(DimNum == EncodedPoseMeans.Num());
+
+	for (int32 FrameIdx = 0; FrameIdx < FrameNum; FrameIdx++)
+	{
+		for (int32 DimIdx = 0; DimIdx < DimNum; DimIdx++)
+		{
+			OutUnnormalizedEncodedPoseVectors[FrameIdx][DimIdx] = InNormalizedEncodedPoseVectors[FrameIdx][DimIdx] * EncodedPoseNormalizationScale + EncodedPoseMeans[DimIdx];
+		}
+	}
+}
+
+void UAnimGenController::ClampNormalizedPoseVectorsInplace(
+	const TLearningArrayView<2, float> InOutNormalizedEncodedPoseVectors) const
+{
+	const int32 FrameNum = InOutNormalizedEncodedPoseVectors.Num<0>();
+	const int32 DimNum = InOutNormalizedEncodedPoseVectors.Num<1>();
+	check(DimNum == EncodedPoseMeans.Num());
+	check(DimNum == EncodedPoseMins.Num());
+	check(DimNum == EncodedPoseMaxs.Num());
+
+	for (int32 FrameIdx = 0; FrameIdx < FrameNum; FrameIdx++)
+	{
+		for (int32 DimIdx = 0; DimIdx < DimNum; DimIdx++)
+		{
+			InOutNormalizedEncodedPoseVectors[FrameIdx][DimIdx] = FMath::Clamp(
+				InOutNormalizedEncodedPoseVectors[FrameIdx][DimIdx],
+				(EncodedPoseMins[DimIdx] - EncodedPoseMeans[DimIdx]) / FMath::Max(EncodedPoseNormalizationScale, UE_SMALL_NUMBER),
+				(EncodedPoseMaxs[DimIdx] - EncodedPoseMeans[DimIdx]) / FMath::Max(EncodedPoseNormalizationScale, UE_SMALL_NUMBER));
+		}
+	}
+}
+
+void UAnimGenController::ScaleSamplingNoiseInplace(const TLearningArrayView<2, float> InOutSamplingNoise) const
+{
+	const int32 FrameNum = InOutSamplingNoise.Num<0>();
+	const int32 DimNum = InOutSamplingNoise.Num<1>();
+	check(DimNum == NormalizedPoseStds.Num());
+
+	for (int32 FrameIdx = 0; FrameIdx < FrameNum; FrameIdx++)
+	{
+		for (int32 DimIdx = 0; DimIdx < DimNum; DimIdx++)
+		{
+			InOutSamplingNoise[FrameIdx][DimIdx] *= NormalizedPoseStds[DimIdx];
+		}
+	}
+}
+
+#undef LOCTEXT_NAMESPACE
