AnimGen: Move to Experimental

diff --git a/Engine/Plugins/Experimental/Animation/AnimGen/Source/AnimGen/Private/AnimGenControllerTrait.cpp b/Engine/Plugins/Experimental/Animation/AnimGen/Source/AnimGen/Private/AnimGenControllerTrait.cpp
new file mode 100644
--- /dev/null
+++ b/Engine/Plugins/Experimental/Animation/AnimGen/Source/AnimGen/Private/AnimGenControllerTrait.cpp
@@ -0,0 +1,666 @@
+// Copyright Epic Games, Inc. All Rights Reserved.
+
+#include "AnimGenControllerTrait.h"
+
+#include "AnimDatabaseMath.h"
+#include "AnimGenBehavior.h"
+#include "AnimGenAutoEncoder.h"
+#include "AnimGenLog.h"
+#include "AnimGenTasks.h"
+#include "EvaluationVM/Tasks/PushReferenceKeyframe.h"
+#include "LearningNeuralNetwork.h"
+#include "LearningRandom.h"
+#include "DrawDebugLibrary.h"
+#include "IPoseHistory.h"
+#include "PoseSearch/PoseSearchHistory.h"
+#include "Animation/AnimTypes.h"
+#include "Engine/SkeletalMesh.h"
+
+#include "AnimationRuntime.h"
+#include "TraitCore/ExecutionContext.h"
+
+namespace UE::UAF
+{
+	AUTO_REGISTER_ANIM_TRAIT(FAnimGenControllerTrait)
+
+	// Trait implementation boilerplate
+	#define TRAIT_INTERFACE_ENUMERATOR(GeneratorMacro) \
+		GeneratorMacro(IEvaluate) \
+		GeneratorMacro(IUpdate) \
+		GeneratorMacro(INotifySource) \
+		GeneratorMacro(IGarbageCollection) \
+
+	GENERATE_ANIM_TRAIT_IMPLEMENTATION(FAnimGenControllerTrait, TRAIT_INTERFACE_ENUMERATOR, NULL_ANIM_TRAIT_INTERFACE_ENUMERATOR, NULL_ANIM_TRAIT_EVENT_ENUMERATOR)
+	#undef TRAIT_INTERFACE_ENUMERATOR
+
+	void FAnimGenControllerTrait::OnBecomeRelevant(FUpdateTraversalContext& Context, const TTraitBinding<IUpdate>& Binding, const FTraitUpdateState& TraitState) const
+	{
+		IUpdate::OnBecomeRelevant(Context, Binding, TraitState);
+
+		FInstanceData* InstanceData = Binding.GetInstanceData<FInstanceData>();
+		check(InstanceData);
+
+		InstanceData->bPoseStateRequiresReset = true;
+	}
+
+	void FAnimGenControllerTrait::PostUpdate(FUpdateTraversalContext& Context, const TTraitBinding<IUpdate>& Binding, const FTraitUpdateState& TraitState) const
+	{
+		IUpdate::PostUpdate(Context, Binding, TraitState);
+
+		const FSharedData* SharedData = Binding.GetSharedData<FSharedData>();
+		FInstanceData* InstanceData = Binding.GetInstanceData<FInstanceData>();
+		const float DeltaTime = TraitState.GetDeltaTime();
+
+		// Update Controller
+
+		UAnimGenController* InputController = SharedData->GetController(Binding);
+		if (!InputController || !InputController->IsValid()) { InstanceData->ActiveController = nullptr; return; }
+
+		// Controller has changed so reset internal state
+
+		if (InputController != InstanceData->ActiveController || !InstanceData->PoseState.IsValid())
+		{
+			// Update Active Controller
+
+			InstanceData->ActiveController = InputController;
+
+			// Get Control Schema
+
+			const FAnimGenControlSchema& ControlSchema = InstanceData->ActiveController->ControlSchema;
+			const FAnimGenControlSchemaElement& ControlSchemaElement = InstanceData->ActiveController->ControlSchemaElement;
+			check(ControlSchema.ObservationSchema.IsValid() && ControlSchema.ObservationSchema->Schema.IsValid(ControlSchemaElement.SchemaElement));
+
+			InstanceData->ControlVectorSize = ControlSchema.ObservationSchema->Schema.GetObservationVectorSize(ControlSchemaElement.SchemaElement);
+			InstanceData->EncodedControlVectorSize = ControlSchema.ObservationSchema->Schema.GetEncodedVectorSize(ControlSchemaElement.SchemaElement);
+			InstanceData->ControlDistributionVectorSize = ControlSchema.ObservationSchema->Schema.GetObservationDistributionVectorSize(ControlSchemaElement.SchemaElement);
+
+			InstanceData->ControlVector.SetNumUninitialized({ 1, InstanceData->ControlVectorSize });
+			InstanceData->EncodedControlVector.SetNumUninitialized({ 1, InstanceData->EncodedControlVectorSize });
+			UE::Learning::Array::Zero(InstanceData->ControlVector);
+			UE::Learning::Array::Zero(InstanceData->EncodedControlVector);
+
+			// Get Networks
+
+			check(InstanceData->ActiveController->ControlEncoderNetwork);
+			check(InstanceData->ActiveController->LOD0Network);
+			check(InstanceData->ActiveController->LOD1Network);
+			check(InstanceData->ActiveController->LOD2Network);
+
+			InstanceData->ControlEncoderInference = InstanceData->ActiveController->ControlEncoderNetwork->GetNetwork()->CreateInferenceObject(1);
+			InstanceData->LOD0Inference = InstanceData->ActiveController->LOD0Network->GetNetwork()->CreateInferenceObject(1);
+			InstanceData->LOD1Inference = InstanceData->ActiveController->LOD1Network->GetNetwork()->CreateInferenceObject(1);
+			InstanceData->LOD2Inference = InstanceData->ActiveController->LOD2Network->GetNetwork()->CreateInferenceObject(1);
+
+			UAnimGenAutoEncoder* AutoEncoder = InstanceData->ActiveController->AutoEncoder;
+			check(AutoEncoder);
+			check(AutoEncoder->EncoderNetwork);
+			check(AutoEncoder->DecoderNetwork);
+
+			InstanceData->EncoderInference = AutoEncoder->EncoderNetwork->GetNetwork()->CreateInferenceObject(1);
+			InstanceData->DecoderInference = AutoEncoder->DecoderNetwork->GetNetwork()->CreateInferenceObject(1);
+
+			// Reset Pose State
+
+			const int32 BoneNum = AutoEncoder->GetBoneNum();
+
+			InstanceData->PoseState = UAnimDatabasePoseStateLibrary::MakePoseState();
+
+			InstanceData->PoseState.PrevPoseData->RootData.Resize(1);
+			InstanceData->PoseState.PrevPoseData->LocalBoneData.Resize(1, BoneNum);
+			InstanceData->PoseState.PrevPoseData->AttributeData.Resize(1, AutoEncoder->GetAttributeTypes());
+			AutoEncoder->SetDefaultPoseData(InstanceData->PoseState.PrevPoseData->View());
+
+			InstanceData->PoseState.CurrPoseData->RootData.Resize(1);
+			InstanceData->PoseState.CurrPoseData->LocalBoneData.Resize(1, BoneNum);
+			InstanceData->PoseState.CurrPoseData->AttributeData.Resize(1, AutoEncoder->GetAttributeTypes());
+			AutoEncoder->SetDefaultPoseData(InstanceData->PoseState.CurrPoseData->View());
+
+			InstanceData->PoseState.PoseGlobalBoneData->Resize(1, BoneNum);
+
+			InstanceData->OutputPoseData.Resize(1, BoneNum, AutoEncoder->GetAttributeTypes());
+			AutoEncoder->SetDefaultPoseData(InstanceData->OutputPoseData.View());
+
+			// Compute Pose Vector
+
+			InstanceData->EncodedPoseVectorSize = AutoEncoder->GetEncodingSize();
+			InstanceData->PoseVectorSize = AutoEncoder->GetPoseVectorSize();
+
+			InstanceData->EncodedPoseVector.SetNumUninitialized({ 1, InstanceData->EncodedPoseVectorSize });
+			InstanceData->NextEncodedPoseVector.SetNumUninitialized({ 1, InstanceData->EncodedPoseVectorSize });
+			InstanceData->UnnormalizedEncodedPoseVector.SetNumUninitialized({ 1, InstanceData->EncodedPoseVectorSize });
+			InstanceData->PoseVector.SetNumUninitialized({ 1, InstanceData->PoseVectorSize });
+			InstanceData->DenoiserInputVector.SetNumUninitialized({ 1, InstanceData->EncodedPoseVectorSize * 2 + 1 + InstanceData->EncodedControlVectorSize });
+			UE::Learning::Array::Zero(InstanceData->EncodedPoseVector);
+			UE::Learning::Array::Zero(InstanceData->NextEncodedPoseVector);
+			UE::Learning::Array::Zero(InstanceData->UnnormalizedEncodedPoseVector);
+			UE::Learning::Array::Zero(InstanceData->PoseVector);
+			UE::Learning::Array::Zero(InstanceData->DenoiserInputVector);
+
+			// Pose Needs Reset
+
+			InstanceData->bPoseStateRequiresReset = true;
+		}
+
+		// Check Controls are Valid
+
+		const FAnimGenControlSchema& ControlSchema = InstanceData->ActiveController->ControlSchema;
+		const FAnimGenControlSchemaElement& ControlSchemaElement = InstanceData->ActiveController->ControlSchemaElement;
+		check(ControlSchema.ObservationSchema.IsValid() && ControlSchema.ObservationSchema->Schema.IsValid(ControlSchemaElement.SchemaElement));
+
+		const FAnimGenControlObject& ControlObject = SharedData->GetControlObject(Binding);
+		const FAnimGenControlObjectElement& ControlObjectElement = SharedData->GetControlObjectElement(Binding);
+
+		if (!UAnimGenControls::ValidateControlObjectMatchesSchema(ControlSchema, ControlSchemaElement, ControlObject, ControlObjectElement))
+		{
+			if (ControlObject.IsValid())
+			{
+				FWriteScopeLock ObjectLock(ControlObject.ObservationObject->Lock);
+				ControlObject.ObservationObject->Object.Reset();
+			}
+
+			InstanceData->bPoseStateRequiresReset = true;
+			InstanceData->bControlsValid = false;
+			return;
+		}
+
+		InstanceData->bControlsValid = true;
+
+		// Set Control Vector
+
+		{
+			FReadScopeLock ObjectLock(ControlObject.ObservationObject->Lock);
+
+			UE::Learning::Observation::SetVectorFromObject(
+				InstanceData->ControlVector[0],
+				ControlSchema.ObservationSchema->Schema,
+				ControlSchemaElement.SchemaElement,
+				ControlObject.ObservationObject->Object,
+				ControlObjectElement.ObjectElement);
+		}
+
+		// Reset Control Object
+
+		{
+			FWriteScopeLock ObjectLock(ControlObject.ObservationObject->Lock);
+			ControlObject.ObservationObject->Object.Reset();
+		}
+
+		// Clamp Control Vectors
+
+		UAnimGenController* Controller = InstanceData->ActiveController;
+		check(Controller);
+
+		UE::Learning::Observation::ClampToObservationDistributionMinMax(
+			InstanceData->ControlVector[0],
+			Controller->ControlVectorDistributionMins,
+			Controller->ControlVectorDistributionMaxs,
+			ControlSchema.ObservationSchema->Schema,
+			ControlSchemaElement.SchemaElement);
+
+		// Get Controller Data
+
+		check(InstanceData->ActiveController->IsValid());
+		check(InstanceData->PoseState.IsValid());
+
+		UAnimGenAutoEncoder* AutoEncoder = InstanceData->ActiveController->AutoEncoder;
+
+		check(AutoEncoder);
+		check(InstanceData->PoseState.IsValid());
+
+		const UE::PoseSearch::IPoseHistory* PoseHistory = nullptr;
+
+		TTraitBinding<UE::UAF::IPoseHistory> PoseHistoryTrait;
+		if (Context.GetScopedInterface<UE::UAF::IPoseHistory>(PoseHistoryTrait))
+		{
+			PoseHistory = PoseHistoryTrait.GetPoseHistory(Context);
+		}
+
+		// Reset Pose
+
+		if (InstanceData->bPoseStateRequiresReset)
+		{
+			// Set to default
+
+			AutoEncoder->SetDefaultPoseData(InstanceData->PoseState.PrevPoseData->View());
+			AutoEncoder->SetDefaultPoseData(InstanceData->PoseState.CurrPoseData->View());
+
+			// Reset Using Pose History
+
+			if (PoseHistory)
+			{
+				USkeleton* Skeleton = 
+					Context.GetBindingObject().Get() && 
+					Context.GetBindingObject()->GetSkeletalMeshAsset() ?
+					Context.GetBindingObject()->GetSkeletalMeshAsset()->GetSkeleton() : nullptr;
+
+				// Reset Root
+
+				FTransform CurrRootTransform = FTransform::Identity;
+
+				// TODO: Get Skeleton
+				if (PoseHistory->GetTransformAtTime(0.0f, CurrRootTransform, Skeleton, 0, UE::PoseSearch::WorldSpaceIndexType))
+				{
+					InstanceData->PoseState.CurrPoseData->RootData.RootLocations[0] = CurrRootTransform.GetLocation();
+					InstanceData->PoseState.CurrPoseData->RootData.RootRotations[0] = ((FQuat4f)CurrRootTransform.GetRotation()).GetNormalized();
+					InstanceData->PoseState.CurrPoseData->RootData.RootScales[0] = (FVector3f)CurrRootTransform.GetScale3D();
+
+					FTransform PrevRootTransform = FTransform::Identity;
+
+					if (DeltaTime > UE_SMALL_NUMBER &&
+						PoseHistory->GetTransformAtTime(-DeltaTime, PrevRootTransform, Skeleton, 0, UE::PoseSearch::WorldSpaceIndexType))
+					{
+						const FVector3f RootLinearVelocity = (FVector3f)(CurrRootTransform.GetLocation() - PrevRootTransform.GetLocation()) / DeltaTime;
+						const FVector3f RootAngularVelocity = (FVector3f)(CurrRootTransform.GetRotation() * PrevRootTransform.GetRotation().Inverse()).GetShortestArcWith(FQuat::Identity).ToRotationVector() / DeltaTime;
+						const FVector3f RootScalarVelocity = UE::AnimDatabase::Math::VectorLogSafe(UE::AnimDatabase::Math::VectorDivMax((FVector3f)CurrRootTransform.GetScale3D(), (FVector3f)PrevRootTransform.GetScale3D())) / DeltaTime;
+
+						InstanceData->PoseState.PrevPoseData->RootData.RootLocations[0] = PrevRootTransform.GetLocation();
+						InstanceData->PoseState.PrevPoseData->RootData.RootRotations[0] = ((FQuat4f)PrevRootTransform.GetRotation()).GetNormalized();
+						InstanceData->PoseState.PrevPoseData->RootData.RootScales[0] = (FVector3f)PrevRootTransform.GetScale3D();
+
+						InstanceData->PoseState.CurrPoseData->RootData.RootLinearVelocities[0] = RootLinearVelocity;
+						InstanceData->PoseState.CurrPoseData->RootData.RootAngularVelocities[0] = RootAngularVelocity;
+						InstanceData->PoseState.CurrPoseData->RootData.RootScalarVelocities[0] = RootScalarVelocity;
+
+						InstanceData->PoseState.PrevPoseData->RootData.RootLinearVelocities[0] = RootLinearVelocity;
+						InstanceData->PoseState.PrevPoseData->RootData.RootAngularVelocities[0] = RootAngularVelocity;
+						InstanceData->PoseState.PrevPoseData->RootData.RootScalarVelocities[0] = RootScalarVelocity;
+					}
+					else
+					{
+						InstanceData->PoseState.PrevPoseData->RootData.RootLocations[0] = CurrRootTransform.GetLocation();
+						InstanceData->PoseState.PrevPoseData->RootData.RootRotations[0] = ((FQuat4f)CurrRootTransform.GetRotation()).GetNormalized();
+						InstanceData->PoseState.PrevPoseData->RootData.RootScales[0] = (FVector3f)CurrRootTransform.GetScale3D();
+					}
+				}
+				else
+				{
+					const FTransform ComponentRootTransform = SharedData->GetMeshTransform(Binding);
+
+					InstanceData->PoseState.CurrPoseData->RootData.RootLocations[0] = ComponentRootTransform.GetLocation();
+					InstanceData->PoseState.CurrPoseData->RootData.RootRotations[0] = ((FQuat4f)ComponentRootTransform.GetRotation()).GetNormalized();
+					InstanceData->PoseState.CurrPoseData->RootData.RootScales[0] = (FVector3f)ComponentRootTransform.GetScale3D();
+
+					InstanceData->PoseState.PrevPoseData->RootData.RootLocations[0] = ComponentRootTransform.GetLocation();
+					InstanceData->PoseState.PrevPoseData->RootData.RootRotations[0] = ((FQuat4f)ComponentRootTransform.GetRotation()).GetNormalized();
+					InstanceData->PoseState.PrevPoseData->RootData.RootScales[0] = (FVector3f)ComponentRootTransform.GetScale3D();
+				}
+
+				// Reset Bones
+
+				const int32 BoneNum = AutoEncoder->GetBoneNum();
+
+				// We iterate from 1 to ignore "root" bone since that will be adjusted by root offset bone
+				for (int32 BoneIdx = 1; BoneIdx < BoneNum; BoneIdx++)
+				{
+					const int32 ParentIdx = AutoEncoder->GetBoneParents()[BoneIdx];
+					if (ParentIdx == INDEX_NONE) { continue; }
+
+					if (!Skeleton) { continue; }
+					const int32 SkeletonBoneIdx = Skeleton->GetReferenceSkeleton().FindBoneIndex(InstanceData->ActiveController->AutoEncoder->GetBoneName(BoneIdx));
+					const int32 SkeletonParentIdx = Skeleton->GetReferenceSkeleton().FindBoneIndex(InstanceData->ActiveController->AutoEncoder->GetBoneName(ParentIdx));
+					if (SkeletonBoneIdx == INDEX_NONE || SkeletonParentIdx == INDEX_NONE) { continue; }
+
+					FTransform CurrBoneTransform = FTransform::Identity;
+					FTransform CurrParentTransform = FTransform::Identity;
+
+					if (PoseHistory->GetTransformAtTime(0.0f, CurrBoneTransform, Skeleton, SkeletonBoneIdx) &&
+						PoseHistory->GetTransformAtTime(0.0f, CurrParentTransform, Skeleton, SkeletonParentIdx))
+					{
+						FTransform CurrTransform = CurrBoneTransform * CurrParentTransform.Inverse();
+
+						InstanceData->PoseState.CurrPoseData->LocalBoneData.BoneLocations[0][BoneIdx] = (FVector3f)CurrTransform.GetLocation();
+						InstanceData->PoseState.CurrPoseData->LocalBoneData.BoneRotations[0][BoneIdx] = (FQuat4f)CurrTransform.GetRotation();
+						InstanceData->PoseState.CurrPoseData->LocalBoneData.BoneScales[0][BoneIdx] = (FVector3f)CurrTransform.GetScale3D();
+
+						FTransform PrevBoneTransform = FTransform::Identity;
+						FTransform PrevParentTransform = FTransform::Identity;
+
+						if (DeltaTime > UE_SMALL_NUMBER &&
+							PoseHistory->GetTransformAtTime(-DeltaTime, PrevBoneTransform, Skeleton, SkeletonBoneIdx) &&
+							PoseHistory->GetTransformAtTime(-DeltaTime, PrevParentTransform, Skeleton, SkeletonParentIdx))
+						{
+							FTransform PrevTransform = PrevBoneTransform * PrevParentTransform.Inverse();
+
+							const FVector3f LinearVelocity = (FVector3f)(CurrTransform.GetLocation() - PrevTransform.GetLocation()) / DeltaTime;
+							const FVector3f AngularVelocity = (FVector3f)(CurrTransform.GetRotation() * PrevTransform.GetRotation().Inverse()).GetShortestArcWith(FQuat::Identity).ToRotationVector() / DeltaTime;
+							const FVector3f ScalarVelocity = UE::AnimDatabase::Math::VectorLogSafe(UE::AnimDatabase::Math::VectorDivMax((FVector3f)CurrTransform.GetScale3D(), (FVector3f)PrevTransform.GetScale3D())) / DeltaTime;
+
+							InstanceData->PoseState.PrevPoseData->LocalBoneData.BoneLocations[0][BoneIdx] = (FVector3f)PrevTransform.GetLocation();
+							InstanceData->PoseState.PrevPoseData->LocalBoneData.BoneRotations[0][BoneIdx] = (FQuat4f)PrevTransform.GetRotation();
+							InstanceData->PoseState.PrevPoseData->LocalBoneData.BoneScales[0][BoneIdx] = (FVector3f)PrevTransform.GetScale3D();
+
+							InstanceData->PoseState.CurrPoseData->LocalBoneData.BoneLinearVelocities[0][BoneIdx] = LinearVelocity;
+							InstanceData->PoseState.CurrPoseData->LocalBoneData.BoneAngularVelocities[0][BoneIdx] = AngularVelocity;
+							InstanceData->PoseState.CurrPoseData->LocalBoneData.BoneScalarVelocities[0][BoneIdx] = ScalarVelocity;
+
+							InstanceData->PoseState.PrevPoseData->LocalBoneData.BoneLinearVelocities[0][BoneIdx] = LinearVelocity;
+							InstanceData->PoseState.PrevPoseData->LocalBoneData.BoneAngularVelocities[0][BoneIdx] = AngularVelocity;
+							InstanceData->PoseState.PrevPoseData->LocalBoneData.BoneScalarVelocities[0][BoneIdx] = ScalarVelocity;
+						}
+						else
+						{
+							InstanceData->PoseState.PrevPoseData->LocalBoneData.BoneLocations[0][BoneIdx] = (FVector3f)CurrTransform.GetLocation();
+							InstanceData->PoseState.PrevPoseData->LocalBoneData.BoneRotations[0][BoneIdx] = (FQuat4f)CurrTransform.GetRotation();
+							InstanceData->PoseState.PrevPoseData->LocalBoneData.BoneScales[0][BoneIdx] = (FVector3f)CurrTransform.GetScale3D();
+						}
+					}
+				}
+			}
+			else
+			{
+				const FTransform ComponentRootTransform = SharedData->GetMeshTransform(Binding);
+
+				InstanceData->PoseState.CurrPoseData->RootData.RootLocations[0] = ComponentRootTransform.GetLocation();
+				InstanceData->PoseState.CurrPoseData->RootData.RootRotations[0] = ((FQuat4f)ComponentRootTransform.GetRotation()).GetNormalized();
+				InstanceData->PoseState.CurrPoseData->RootData.RootScales[0] = (FVector3f)ComponentRootTransform.GetScale3D();
+
+				InstanceData->PoseState.PrevPoseData->RootData.RootLocations[0] = ComponentRootTransform.GetLocation();
+				InstanceData->PoseState.PrevPoseData->RootData.RootRotations[0] = ((FQuat4f)ComponentRootTransform.GetRotation()).GetNormalized();
+				InstanceData->PoseState.PrevPoseData->RootData.RootScales[0] = (FVector3f)ComponentRootTransform.GetScale3D();
+			}
+
+			AutoEncoder->ToPoseVectors(InstanceData->PoseVector, InstanceData->PoseState.CurrPoseData->ConstView());
+			AutoEncoder->ClampPoseVectors(InstanceData->PoseVector);
+			AutoEncoder->NormalizePoseVectors(InstanceData->PoseVector);
+
+			// Encode Pose Vector
+
+			if (InstanceData->EncoderInference->GetInputSize() == InstanceData->PoseVector.Num<1>() &&
+				InstanceData->EncoderInference->GetOutputSize() == InstanceData->UnnormalizedEncodedPoseVector.Num<1>())
+			{
+				InstanceData->EncoderInference->Evaluate(InstanceData->UnnormalizedEncodedPoseVector, InstanceData->PoseVector);
+			}
+			else
+			{
+				UE::Learning::Array::Zero(InstanceData->UnnormalizedEncodedPoseVector);
+			}
+
+			// Normalize Encoded Pose Vector
+
+			Controller->NormalizeEncodedPoseVectors(InstanceData->EncodedPoseVector, InstanceData->UnnormalizedEncodedPoseVector);
+			Controller->ClampNormalizedPoseVectorsInplace(InstanceData->EncodedPoseVector);
+			UE::Learning::Array::Copy(InstanceData->NextEncodedPoseVector, InstanceData->EncodedPoseVector);
+		}
+
+		// Update Pose State
+
+		if (SharedData->bRootFromPoseHistory && PoseHistory)
+		{
+			USkeleton* Skeleton = Context.GetBindingObject().Get() && Context.GetBindingObject()->GetSkeletalMeshAsset() ?
+				Context.GetBindingObject()->GetSkeletalMeshAsset()->GetSkeleton() : nullptr;
+
+			FTransform PoseRootTransform = FTransform::Identity;
+			if (PoseHistory->GetTransformAtTime(0.0f, PoseRootTransform, Skeleton, 0, UE::PoseSearch::WorldSpaceIndexType))
+			{
+				UE::AnimDatabase::PoseData::Teleport(
+					InstanceData->PoseState.CurrPoseData->RootData.View(),
+					{ PoseRootTransform.GetLocation() },
+					{ (FQuat4f)PoseRootTransform.GetRotation() });
+			}
+		}
+
+#if WITH_EDITOR
+		if (SharedData->GetbDrawDebugPoseState(Binding))
+		{
+			UDrawDebugLibrary::DrawDebugRotation(
+				FDebugDrawer::MakeDebugDrawer(Context.GetDebugDrawInterface()),
+				InstanceData->PoseState.CurrPoseData->RootData.RootLocations[0],
+				((FQuat)InstanceData->PoseState.CurrPoseData->RootData.RootRotations[0]).GetNormalized().Rotator(),
+				UDrawDebugLibrary::MakeDrawDebugLineStyleFromColorAndThickness(FLinearColor::White, 2.0f),
+				false,
+				10.0f);
+
+			UE::AnimDatabase::PoseData::ForwardKinematics(
+				InstanceData->PoseState.PoseGlobalBoneData->View(),
+				InstanceData->PoseState.CurrPoseData->LocalBoneData.ConstView(),
+				InstanceData->PoseState.CurrPoseData->RootData.ConstView(),
+				AutoEncoder->GetBoneParents());
+
+			UDrawDebugLibrary::DrawDebugSkeletonArrayView(
+				FDebugDrawer::MakeDebugDrawer(Context.GetDebugDrawInterface()),
+				InstanceData->PoseState.CurrPoseData->RootData.RootLocations[0],
+				InstanceData->PoseState.CurrPoseData->RootData.RootRotations[0],
+				InstanceData->PoseState.PoseGlobalBoneData->BoneLocations[0],
+				InstanceData->PoseState.PoseGlobalBoneData->BoneRotations[0],
+				AutoEncoder->AutoEncodedRequiredBoneIndices,
+				AutoEncoder->GetBoneParents());
+		}
+#endif
+
+		// Evaluate Control Encoder
+
+		if (InstanceData->ControlVector.Num<1>() == InstanceData->ControlEncoderInference->GetInputSize() &&
+			InstanceData->EncodedControlVector.Num<1>() == InstanceData->ControlEncoderInference->GetOutputSize())
+		{
+			InstanceData->ControlEncoderInference->Evaluate(InstanceData->EncodedControlVector, InstanceData->ControlVector);
+		}
+		else
+		{
+			UE::Learning::Array::Zero(InstanceData->EncodedControlVector);
+		}
+
+		// Evaluate Denoiser
+
+		const float ControllerDeltaTime = 1.0 / FMath::Max(Controller->FrameRate.AsDecimal(), UE_SMALL_NUMBER);
+		const float SubsampleRate = FMath::Clamp(FMath::Max(DeltaTime, ControllerDeltaTime) * Controller->FrameRate.AsDecimal(), 1, 8);
+
+		UE::Learning::Array::Copy(InstanceData->NextEncodedPoseVector, InstanceData->EncodedPoseVector);
+
+		if (InstanceData->DenoiserInputVector.Num<1>() == InstanceData->LOD0Inference->GetInputSize() &&
+			InstanceData->DenoiserInputVector.Num<1>() == InstanceData->LOD1Inference->GetInputSize() &&
+			InstanceData->DenoiserInputVector.Num<1>() == InstanceData->LOD2Inference->GetInputSize() &&
+			InstanceData->EncodedPoseVectorSize == InstanceData->LOD0Inference->GetOutputSize() &&
+			InstanceData->EncodedPoseVectorSize == InstanceData->LOD1Inference->GetOutputSize() &&
+			InstanceData->EncodedPoseVectorSize == InstanceData->LOD2Inference->GetOutputSize())
+		{
+			UE::Learning::Random::SampleClippedGaussianArray(InstanceData->NextEncodedPoseVector[0], InstanceData->RandomState, 0.0f, 1.0f, 3.0f);
+			Controller->ScaleSamplingNoiseInplace(InstanceData->NextEncodedPoseVector);
+
+			UE::Learning::Array::Copy(InstanceData->DenoiserInputVector[0].Slice(0, InstanceData->EncodedPoseVectorSize), InstanceData->EncodedPoseVector[0]);
+			UE::Learning::Array::Copy(InstanceData->DenoiserInputVector[0].Slice(InstanceData->EncodedPoseVectorSize, InstanceData->EncodedPoseVectorSize), InstanceData->NextEncodedPoseVector[0]);
+			UE::Learning::Array::Set(InstanceData->DenoiserInputVector[0].Slice(InstanceData->EncodedPoseVectorSize * 2, 1), (SubsampleRate - 1.0f) / 7.0f);
+			UE::Learning::Array::Copy(InstanceData->DenoiserInputVector[0].Slice(InstanceData->EncodedPoseVectorSize * 2 + 1, InstanceData->EncodedControlVectorSize), InstanceData->EncodedControlVector[0]);
+
+			switch (SharedData->GetLODLevel(Binding))
+			{
+			default:
+			case 0: InstanceData->LOD0Inference->Evaluate(InstanceData->NextEncodedPoseVector, InstanceData->DenoiserInputVector); break;
+			case 1: InstanceData->LOD1Inference->Evaluate(InstanceData->NextEncodedPoseVector, InstanceData->DenoiserInputVector); break;
+			case 2: InstanceData->LOD2Inference->Evaluate(InstanceData->NextEncodedPoseVector, InstanceData->DenoiserInputVector); break;
+			}
+
+			Controller->ClampNormalizedPoseVectorsInplace(InstanceData->NextEncodedPoseVector);
+		}
+		else
+		{
+			UE::Learning::Array::Zero(InstanceData->NextEncodedPoseVector);
+		}
+
+		// Interpolate / Extrapolate using Delta Time Ratio
+
+		if (DeltaTime < ControllerDeltaTime)
+		{
+			UE::AnimDatabase::Math::LerpToTargetInplace(
+				InstanceData->EncodedPoseVector,
+				InstanceData->NextEncodedPoseVector,
+				DeltaTime * Controller->FrameRate.AsDecimal());
+		}
+		else
+		{
+			UE::Learning::Array::Copy(InstanceData->EncodedPoseVector, InstanceData->NextEncodedPoseVector);
+		}
+
+		Controller->ClampNormalizedPoseVectorsInplace(InstanceData->EncodedPoseVector);
+
+		// Evaluate Decoder
+
+		Controller->DenormalizeEncodedPoseVectors(InstanceData->UnnormalizedEncodedPoseVector, InstanceData->EncodedPoseVector);
+
+		if (InstanceData->DecoderInference->GetInputSize() == InstanceData->UnnormalizedEncodedPoseVector.Num<1>() &&
+			InstanceData->DecoderInference->GetOutputSize() == InstanceData->PoseVector.Num<1>())
+		{
+			InstanceData->DecoderInference->Evaluate(InstanceData->PoseVector, InstanceData->UnnormalizedEncodedPoseVector);
+		}
+		else
+		{
+			UE::Learning::Array::Zero(InstanceData->PoseVector);
+		}
+
+		// Update Prev Pose Data
+
+		*InstanceData->PoseState.PrevPoseData = *InstanceData->PoseState.CurrPoseData;
+
+		// Update Curr Pose Data
+
+		const FTransform PoseRootTransform = SharedData->GetMeshTransform(Binding);
+
+		AutoEncoder->DenormalizePoseVectors(InstanceData->PoseVector);
+		AutoEncoder->ClampPoseVectors(InstanceData->PoseVector);
+		AutoEncoder->FromPoseVectors(InstanceData->OutputPoseData.View(), InstanceData->PoseVector,
+			{ PoseRootTransform.GetLocation() },
+			{ ((FQuat4f)PoseRootTransform.GetRotation()).GetNormalized() });
+
+		InstanceData->PoseState.CurrPoseData->RootData = InstanceData->OutputPoseData.RootData;
+		InstanceData->PoseState.CurrPoseData->AttributeData = InstanceData->OutputPoseData.AttributeData;
+
+		if (InstanceData->bPoseStateRequiresReset)
+		{
+			InstanceData->PoseState.CurrPoseData->LocalBoneData = InstanceData->OutputPoseData.LocalBoneData;
+		}
+		else
+		{
+			UE::AnimDatabase::PoseData::DampInplaceUsingAdaptiveVelocityIntegrationMix(
+				InstanceData->PoseState.CurrPoseData->LocalBoneData.View(),
+				InstanceData->OutputPoseData.LocalBoneData.ConstView(),
+				DeltaTime,
+				SharedData->GetMinPoseSmoothingTime(Binding),
+				SharedData->GetMaxPoseSmoothingTime(Binding));
+		}
+
+		// Update Delta Time
+
+		InstanceData->PoseState.DeltaTime = DeltaTime;
+
+		// Update Notifies
+
+		InstanceData->AnimNotifyEvents.Reset();
+
+		const int32 AnimNotifyNum = FMath::Min(SharedData->AnimNotifyAttributeNames.Num(), SharedData->AnimNotifyClasses.Num());
+
+		for (int32 AnimNotifyIdx = 0; AnimNotifyIdx < AnimNotifyNum; AnimNotifyIdx++)
+		{
+			if (!SharedData->AnimNotifyClasses[AnimNotifyIdx]) { continue; }
+
+			const int32 AttributeIdx = AutoEncoder->FindAttributeIndex(SharedData->AnimNotifyAttributeNames[AnimNotifyIdx]);
+			if (AttributeIdx == INDEX_NONE)
+			{
+				continue;
+			}
+
+			if (AutoEncoder->GetAttributeType(AttributeIdx) != EAnimDatabaseAttributeType::Event)
+			{
+				continue;
+			}
+
+			if (!InstanceData->PoseState.PrevPoseData->AttributeData.GetAttributeActive(0, AttributeIdx) ||
+				!InstanceData->PoseState.CurrPoseData->AttributeData.GetAttributeActive(0, AttributeIdx))
+			{
+				continue;
+			}
+
+			bool bPrevTimeUntilEventKnown = false;
+			bool bCurrTimeUntilEventKnown = false;
+			float PrevTimeUntilEvent = UE_MAX_FLT;
+			float CurrTimeUntilEvent = UE_MAX_FLT;
+			InstanceData->PoseState.PrevPoseData->AttributeData.GetEvent(bPrevTimeUntilEventKnown, PrevTimeUntilEvent, 0, AttributeIdx);
+			InstanceData->PoseState.CurrPoseData->AttributeData.GetEvent(bCurrTimeUntilEventKnown, CurrTimeUntilEvent, 0, AttributeIdx);
+
+			if (bPrevTimeUntilEventKnown &&
+				bCurrTimeUntilEventKnown &&
+				PrevTimeUntilEvent > 0.0f &&
+				CurrTimeUntilEvent < 0.0f &&
+				FMath::Abs(PrevTimeUntilEvent - CurrTimeUntilEvent) < 0.1f)
+			{
+				FAnimNotifyEvent& NotifyEvent = InstanceData->AnimNotifyEvents.AddDefaulted_GetRef();
+				NotifyEvent.Notify = SharedData->AnimNotifyClasses[AnimNotifyIdx]->GetDefaultObject<UAnimNotify>();
+				NotifyEvent.NotifyName = AutoEncoder->GetAttributeName(AttributeIdx);
+			}
+		}
+
+		const int32 AnimNotifyStateNum = FMath::Min(SharedData->AnimNotifyStateAttributeNames.Num(), SharedData->AnimNotifyStateClasses.Num());
+
+		for (int32 AnimNotifyStateIdx = 0; AnimNotifyStateIdx < AnimNotifyStateNum; AnimNotifyStateIdx++)
+		{
+			if (!SharedData->AnimNotifyStateClasses[AnimNotifyStateIdx]) { continue; }
+
+			const int32 AttributeIdx = AutoEncoder->FindAttributeIndex(SharedData->AnimNotifyStateAttributeNames[AnimNotifyStateIdx]);
+			if (AttributeIdx == INDEX_NONE)
+			{
+				continue;
+			}
+
+			if (AutoEncoder->GetAttributeType(AttributeIdx) != EAnimDatabaseAttributeType::Bool)
+			{
+				continue;
+			}
+
+			if (!InstanceData->PoseState.CurrPoseData->AttributeData.GetAttributeActive(0, AttributeIdx))
+			{
+				continue;
+			}
+
+			if (InstanceData->PoseState.CurrPoseData->AttributeData.GetBool(0, AttributeIdx))
+			{
+				FAnimNotifyEvent& NotifyEvent = InstanceData->AnimNotifyEvents.AddDefaulted_GetRef();
+				NotifyEvent.NotifyStateClass = SharedData->AnimNotifyStateClasses[AnimNotifyStateIdx]->GetDefaultObject<UAnimNotifyState>();
+				NotifyEvent.NotifyName = AutoEncoder->GetAttributeName(AttributeIdx);
+			}
+		}
+
+		// Pose Reset no longer required
+
+		InstanceData->bPoseStateRequiresReset = false;
+	}
+
+	void FAnimGenControllerTrait::PostEvaluate(FEvaluateTraversalContext& Context, const TTraitBinding<IEvaluate>& Binding) const
+	{
+		FInstanceData* InstanceData = Binding.GetInstanceData<FInstanceData>();
+		const FSharedData* SharedData = Binding.GetSharedData<FSharedData>();
+
+		if (InstanceData->ActiveController && InstanceData->PoseState.IsValid() && InstanceData->bControlsValid)
+		{
+			check(InstanceData->ActiveController->IsValid());
+
+			Context.AppendTask(FAnimGenPushPoseFromPoseState::Make(
+				InstanceData->PoseState, 
+				SharedData->CurveAttributes,
+				InstanceData->ActiveController->AutoEncoder->GetAttributeNames(),
+				InstanceData->ActiveController->AutoEncoder->GetBoneNames()));
+		}
+		else
+		{
+			Context.AppendTask(FAnimNextPushReferenceKeyframeTask::MakeFromSkeleton());
+		}
+	}
+
+	void FAnimGenControllerTrait::AddReferencedObjects(const FExecutionContext& Context, const TTraitBinding<IGarbageCollection>& Binding, FReferenceCollector& Collector) const
+	{
+		IGarbageCollection::AddReferencedObjects(Context, Binding, Collector);
+
+		FInstanceData* InstanceData = Binding.GetInstanceData<FInstanceData>();
+
+		Collector.AddReferencedObject(InstanceData->ActiveController);
+	}
+
+	void FAnimGenControllerTrait::GetNotifies(FExecutionContext& Context, const TTraitBinding<INotifySource>& Binding, float StartPosition, float Duration, bool bLooping, TArray<FAnimNotifyEventReference>& OutNotifies) const
+	{
+		const FInstanceData* InstanceData = Binding.GetInstanceData<FInstanceData>();
+		check(InstanceData);
+
+		const int32 AnimNotifyEventNum = InstanceData->AnimNotifyEvents.Num();
+		for (int32 AnimNotifyEventIdx = 0; AnimNotifyEventIdx < AnimNotifyEventNum; AnimNotifyEventIdx++)
+		{
+			OutNotifies.Emplace(&InstanceData->AnimNotifyEvents[AnimNotifyEventIdx], Context.GetHostObject());
+		}
+	}
+}
diff --git a/Engine/Plugins/Experimental/Animation/AnimGen/Source/AnimGen/Private/AnimGenControllerTrait.h b/Engine/Plugins/Experimental/Animation/AnimGen/Source/AnimGen/Private/AnimGenControllerTrait.h
new file mode 100644
--- /dev/null
+++ b/Engine/Plugins/Experimental/Animation/AnimGen/Source/AnimGen/Private/AnimGenControllerTrait.h
@@ -0,0 +1,179 @@
+// Copyright Epic Games, Inc. All Rights Reserved.
+
+#pragma once
+
+#include "CoreMinimal.h"
+#include "TraitCore/Trait.h"
+#include "TraitInterfaces/IEvaluate.h"
+#include "TraitInterfaces/IGarbageCollection.h"
+#include "TraitInterfaces/IUpdate.h"
+#include "TraitInterfaces/INotifySource.h"
+
+#include "AnimGenController.h"
+
+#include "AnimGenControllerTrait.generated.h"
+
+namespace UE::Learning
+{
+	struct FNeuralNetworkInference;
+}
+
+/** Trait shared data for  the AnimGenController. */
+USTRUCT(meta = (DisplayName = "Anim Gen Controller", ShowTooltip = true))
+struct FAnimNextAnimGenControllerTraitSharedData : public FAnimNextTraitSharedData
+{
+	GENERATED_BODY()
+
+public:
+
+	/** Controller asset used by the trait */
+	UPROPERTY(EditAnywhere, Category = Settings, meta = (PinShownByDefault))
+	TObjectPtr<UAnimGenController> Controller;
+
+	/** Control Object taken as input */
+	UPROPERTY(EditAnywhere, Category = Settings, meta = (PinShownByDefault))
+	FAnimGenControlObject ControlObject;
+
+	/** Control Object Element from Control Object to use as input */
+	UPROPERTY(EditAnywhere, Category = Settings, meta = (PinShownByDefault))
+	FAnimGenControlObjectElement ControlObjectElement;
+
+	/** The world-space transform of the mesh (pre-movement). */
+	UPROPERTY(EditAnywhere, Category = Settings, meta = (PinShownByDefault))
+	FTransform MeshTransform = FTransform::Identity;
+
+	/** Minimum Smoothing time to use - mixes the predicted pose with an integration of the predicted velocity. */
+	UPROPERTY(EditAnywhere, Category = Settings, meta = (PinShownByDefault, ClampMin = "0.0", UIMin = "0.0"))
+	float MinPoseSmoothingTime = 0.05f;
+
+	/** Maximum Smoothing time to use - mixes the predicted pose with an integration of the predicted velocity. */
+	UPROPERTY(EditAnywhere, Category = Settings, meta = (PinShownByDefault, ClampMin = "0.0", UIMin = "0.0"))
+	float MaxPoseSmoothingTime = 0.25f;
+
+	/** Which LOD level to use **/
+	UPROPERTY(EditAnywhere, Category = Settings, meta = (PinShownByDefault, ClampMin = "0", UIMin = "0", ClampMax = "2", UIMax = "2"))
+	int32 LODLevel = 0;
+
+	/** If to draw debug the pose state */
+	UPROPERTY(EditAnywhere, Category = Settings, meta = (PinShownByDefault))
+	bool bDrawDebugPoseState = true;
+
+public:
+
+	/** If to get the root bone from the pose history */
+	UPROPERTY(EditAnywhere, Category = Settings, meta = (Inline))
+	bool bRootFromPoseHistory = true;
+
+	/** Frame Attributes to output as curves */
+	UPROPERTY(EditAnywhere, Category = Settings, meta = (Inline))
+	TArray<FName> CurveAttributes;
+
+	/** Attribute names to anim notify states to produce */
+	UPROPERTY(EditAnywhere, Category = Settings, meta = (Inline))
+	TArray<FName> AnimNotifyStateAttributeNames;
+
+	/** Classes of anim notify states to produce */
+	UPROPERTY(EditAnywhere, Category = Settings, meta = (Inline))
+	TArray<TSubclassOf<UAnimNotifyState>> AnimNotifyStateClasses;
+
+	/** Attribute names to anim notifies to produce */
+	UPROPERTY(EditAnywhere, Category = Settings, meta = (Inline))
+	TArray<FName> AnimNotifyAttributeNames;
+
+	/** Classes of anim notifies to produce */
+	UPROPERTY(EditAnywhere, Category = Settings, meta = (Inline))
+	TArray<TSubclassOf<UAnimNotify>> AnimNotifyClasses;
+
+public:
+
+	// Latent pin support boilerplate
+#define TRAIT_LATENT_PROPERTIES_ENUMERATOR(GeneratorMacro) \
+		GeneratorMacro(Controller) \
+		GeneratorMacro(ControlObject) \
+		GeneratorMacro(ControlObjectElement) \
+		GeneratorMacro(MeshTransform) \
+		GeneratorMacro(MinPoseSmoothingTime) \
+		GeneratorMacro(MaxPoseSmoothingTime) \
+		GeneratorMacro(LODLevel) \
+		GeneratorMacro(bDrawDebugPoseState)
+
+	GENERATE_TRAIT_LATENT_PROPERTIES(FAnimNextAnimGenControllerTraitSharedData, TRAIT_LATENT_PROPERTIES_ENUMERATOR)
+#undef TRAIT_LATENT_PROPERTIES_ENUMERATOR
+};
+
+namespace UE::UAF
+{
+	/**
+	 * FAnimGenControllerTrait
+	 * 
+	 * A trait for generating animation via the AnimGenController
+	 */
+	struct FAnimGenControllerTrait : FBaseTrait, IEvaluate, IUpdate, INotifySource, IGarbageCollection
+	{
+		DECLARE_ANIM_TRAIT(FAnimGenControllerTrait, FBaseTrait)
+
+		using FSharedData = FAnimNextAnimGenControllerTraitSharedData;
+
+		struct FInstanceData : FTrait::FInstanceData
+		{
+			// Random State for Noise
+			
+			uint32 RandomState = 1234;
+
+			// Active Controller Instance
+
+			TObjectPtr<UAnimGenController> ActiveController;
+
+			// Control Variables
+
+			bool bControlsValid = false;
+			int32 ControlVectorSize = 0;
+			int32 EncodedControlVectorSize = 0;
+			int32 ControlDistributionVectorSize = 0;
+			TLearningArray<2, float> ControlVector;
+			TLearningArray<2, float> EncodedControlVector;
+
+			// Pose Variables
+
+			bool bPoseStateRequiresReset = true;
+			FAnimDatabasePoseState PoseState;
+			int32 EncodedPoseVectorSize = 0;
+			int32 PoseVectorSize = 0;
+			TLearningArray<2, float> EncodedPoseVector;
+			TLearningArray<2, float> NextEncodedPoseVector;
+			TLearningArray<2, float> UnnormalizedEncodedPoseVector;
+			TLearningArray<2, float> PoseVector;
+			UE::AnimDatabase::FPoseData OutputPoseData;
+
+			// Network Instances
+
+			TLearningArray<2, float> DenoiserInputVector;
+			TLearningArray<2, float> DenoiserOutputVector;
+			TSharedPtr<UE::Learning::FNeuralNetworkInference> ControlEncoderInference;
+			TSharedPtr<UE::Learning::FNeuralNetworkInference> LOD0Inference;
+			TSharedPtr<UE::Learning::FNeuralNetworkInference> LOD1Inference;
+			TSharedPtr<UE::Learning::FNeuralNetworkInference> LOD2Inference;
+			TSharedPtr<UE::Learning::FNeuralNetworkInference> EncoderInference;
+			TSharedPtr<UE::Learning::FNeuralNetworkInference> DecoderInference;
+
+			/** Internal buffer of AnimNotifyEvents */
+			TArray<FAnimNotifyEvent> AnimNotifyEvents;
+
+			/** Internal buffer of AnimNotifyEventReferences */
+			TArray<FAnimNotifyEventReference> AnimNotifyEventReferences;
+		};
+
+		// IEvaluate impl
+		virtual void PostEvaluate(FEvaluateTraversalContext& Context, const TTraitBinding<IEvaluate>& Binding) const override;
+
+		// IUpdate impl
+		virtual void PostUpdate(FUpdateTraversalContext& Context, const TTraitBinding<IUpdate>& Binding, const FTraitUpdateState& TraitState) const override;
+		virtual void OnBecomeRelevant(FUpdateTraversalContext& Context, const TTraitBinding<IUpdate>& Binding, const FTraitUpdateState& TraitState) const override;
+
+		// IGarbageCollection impl
+		virtual void AddReferencedObjects(const FExecutionContext& Context, const TTraitBinding<IGarbageCollection>& Binding, FReferenceCollector& Collector) const override;
+
+		// INotifySource impl
+		virtual void GetNotifies(FExecutionContext& Context, const TTraitBinding<INotifySource>& Binding, float StartPosition, float Duration, bool bLooping, TArray<FAnimNotifyEventReference>& OutNotifies) const override;
+	};
+}
diff --git a/Engine/Plugins/Experimental/Animation/AnimGen/Source/AnimGen/Private/AnimGenTasks.cpp b/Engine/Plugins/Experimental/Animation/AnimGen/Source/AnimGen/Private/AnimGenTasks.cpp
new file mode 100644
--- /dev/null
+++ b/Engine/Plugins/Experimental/Animation/AnimGen/Source/AnimGen/Private/AnimGenTasks.cpp
@@ -0,0 +1,106 @@
+// Copyright Epic Games, Inc. All Rights Reserved.
+
+#include "AnimGenTasks.h"
+
+#include "Animation/AnimTrace.h"
+#include "EvaluationVM/EvaluationVM.h"
+#include "EvaluationVM/KeyframeState.h"
+
+#include "Animation/AnimRootMotionProvider.h"
+#include "Animation/AnimCurveUtils.h"
+
+#include "AnimDatabasePose.h"
+#include "AnimDatabaseFrameAttribute.h"
+#include "AnimGenLog.h"
+
+#include UE_INLINE_GENERATED_CPP_BY_NAME(AnimGenTasks)
+
+FAnimGenPushPoseFromPoseState FAnimGenPushPoseFromPoseState::Make(
+	const FAnimDatabasePoseState& PoseState, 
+	const TArrayView<const FName> CurveAttributes,
+	const TArrayView<const FName> AttributeNames,
+	const TArrayView<const FName> BoneNames)
+{
+	FAnimGenPushPoseFromPoseState Task;
+	Task.PoseState = PoseState;
+	Task.CurveAttributes = CurveAttributes;
+	Task.AttributeNames = AttributeNames;
+	Task.BoneNames = BoneNames;
+	return Task;
+}
+
+void FAnimGenPushPoseFromPoseState::Execute(UE::UAF::FEvaluationVM& VM) const
+{
+	using namespace UE::UAF;
+
+	FKeyframeState KeyFrame = VM.MakeReferenceKeyframe(false);
+
+	if (!PoseState.IsValid() || PoseState.CurrPoseData->GetFrameNum() != 1)
+	{
+		VM.PushValue(KEYFRAME_STACK_NAME, MakeUnique<FKeyframeState>(KeyFrame));
+		return;
+	}
+
+	const int32 LODBoneNum = KeyFrame.Pose.GetNumBones();
+	const int32 PoseStateBoneNum = PoseState.CurrPoseData->GetBoneNum();
+
+	for (FBoneIndexType LODBoneIdx = 0; LODBoneIdx < LODBoneNum; LODBoneIdx++)
+	{
+		// Get the Skeleton bone associated with each LOD Bone
+		const FBoneIndexType SkeletonBoneIdx = KeyFrame.Pose.GetLODBoneIndexToSkeletonBoneIndexMap()[LODBoneIdx];
+
+		if (SkeletonBoneIdx != INVALID_BONE_INDEX)
+		{
+			// Find the AutoEncoder bone index associated with the skeleton bone
+			const int32 AutoEncoderBoneIdx = BoneNames.Find(KeyFrame.Pose.GetSkeletonAsset()->GetReferenceSkeleton().GetBoneName(SkeletonBoneIdx));
+
+			if (AutoEncoderBoneIdx != INDEX_NONE)
+			{
+				check(AutoEncoderBoneIdx >= 0 && AutoEncoderBoneIdx < PoseStateBoneNum);
+				KeyFrame.Pose.LocalTransformsView.Translations[LODBoneIdx] = (FVector)PoseState.CurrPoseData->LocalBoneData.BoneLocations[0][AutoEncoderBoneIdx];
+				KeyFrame.Pose.LocalTransformsView.Rotations[LODBoneIdx] = ((FQuat)PoseState.CurrPoseData->LocalBoneData.BoneRotations[0][AutoEncoderBoneIdx]).GetNormalized();
+				KeyFrame.Pose.LocalTransformsView.Scales3D[LODBoneIdx] = (FVector)PoseState.CurrPoseData->LocalBoneData.BoneScales[0][AutoEncoderBoneIdx];
+			}
+		}
+	}
+
+	// Write Curves
+
+	TArray<TTuple<FName, float>, TInlineAllocator<32>> OutputCurves;
+	OutputCurves.Reserve(CurveAttributes.Num());
+
+	for (const FName& CurveAttribute : CurveAttributes)
+	{
+		const int32 AttributeIdx = AttributeNames.Find(CurveAttribute);
+
+		if (AttributeIdx == INDEX_NONE) { continue; }
+		if (AttributeIdx < 0 || AttributeIdx >= PoseState.CurrPoseData->AttributeData.GetAttributeNum()) { continue; }
+		if (PoseState.CurrPoseData->AttributeData.GetAttributeType(AttributeIdx) != EAnimDatabaseAttributeType::Float) { continue; }
+		if (PoseState.CurrPoseData->AttributeData.GetAttributeActive(0, AttributeIdx))
+		{
+			const float CurveValue = PoseState.CurrPoseData->AttributeData.GetFloat(0, AttributeIdx);
+			OutputCurves.Add({ CurveAttribute, CurveValue });
+		}
+	}
+
+	UE::Anim::FCurveUtils::BuildUnsorted(KeyFrame.Curves, OutputCurves);
+
+	// Write Root Motion Attribute
+
+	if (const UE::Anim::IAnimRootMotionProvider* RootMotionProvider = UE::Anim::IAnimRootMotionProvider::Get())
+	{
+		const FVector3f LocalRootLinearVelocity = PoseState.CurrPoseData->RootData.RootRotations[0].UnrotateVector(
+			PoseState.CurrPoseData->RootData.RootLinearVelocities[0]);
+
+		const FVector3f LocalRootAngularVelocity = PoseState.CurrPoseData->RootData.RootRotations[0].UnrotateVector(
+			PoseState.CurrPoseData->RootData.RootAngularVelocities[0]);
+
+		const FTransform RootMotionDelta = FTransform(
+			FQuat::MakeFromRotationVector(PoseState.DeltaTime * (FVector)LocalRootAngularVelocity),
+			PoseState.DeltaTime * (FVector)LocalRootLinearVelocity);
+
+		RootMotionProvider->SetRootMotion(RootMotionDelta, KeyFrame.Attributes);
+	}
+
+	VM.PushValue(KEYFRAME_STACK_NAME, MakeUnique<FKeyframeState>(KeyFrame));
+}
\ No newline at end of file
diff --git a/Engine/Plugins/Experimental/Animation/AnimGen/Source/AnimGen/Private/AnimNode_AnimGenController.cpp b/Engine/Plugins/Experimental/Animation/AnimGen/Source/AnimGen/Private/AnimNode_AnimGenController.cpp
new file mode 100644
--- /dev/null
+++ b/Engine/Plugins/Experimental/Animation/AnimGen/Source/AnimGen/Private/AnimNode_AnimGenController.cpp
@@ -0,0 +1,718 @@
+// Copyright Epic Games, Inc. All Rights Reserved.
+
+#include "AnimNode_AnimGenController.h"
+
+#include "AnimGenAutoEncoder.h"
+#include "AnimGenControl.h"
+#include "AnimGenLog.h"
+#include "AnimGenController.h"
+#include "AnimGenBehavior.h"
+
+#include "AnimDatabaseMath.h"
+#include "AnimDatabasePose.h"
+
+#include "LearningRandom.h"
+#include "LearningObservation.h"
+#include "LearningNeuralNetwork.h"
+
+#include "Animation/AnimInstanceProxy.h"
+#include "Animation/AnimCurveUtils.h"
+#include "Animation/AnimBlueprintGeneratedClass.h"
+#include "Animation/AnimRootMotionProvider.h"
+
+#include "PoseSearch/PoseSearchHistory.h"
+#include "PoseSearch/PoseHistoryProvider.h"
+
+#include "DrawDebugLibrary.h"
+
+#include UE_INLINE_GENERATED_CPP_BY_NAME(AnimNode_AnimGenController)
+
+#define LOCTEXT_NAMESPACE "AnimNode_AnimGenController"
+
+void FAnimNode_AnimGenController::Initialize_AnyThread(const FAnimationInitializeContext& Context)
+{
+	DECLARE_SCOPE_HIERARCHICAL_COUNTER_ANIMNODE(Initialize_AnyThread);
+
+	FAnimNode_Base::Initialize_AnyThread(Context);
+}
+
+void FAnimNode_AnimGenController::Update_AnyThread(const FAnimationUpdateContext& Context)
+{
+	DECLARE_SCOPE_HIERARCHICAL_COUNTER_ANIMNODE(Update_AnyThread);
+
+	FAnimNode_Base::Update_AnyThread(Context);
+
+	const float DeltaTime = Context.GetDeltaTime();
+
+	bPoseStateRequiresReset = bPoseStateRequiresReset || (bResetOnBecomingRelevant &&
+		UpdateCounter.HasEverBeenUpdated() &&
+		!UpdateCounter.WasSynchronizedCounter(Context.AnimInstanceProxy->GetUpdateCounter()));
+
+	UpdateCounter.SynchronizeWith(Context.AnimInstanceProxy->GetUpdateCounter());
+
+	GetEvaluateGraphExposedInputs().Execute(Context);
+
+	// Update Active Controller
+
+	if (!Controller || !Controller->IsValid()) { ActiveController = nullptr; return; }
+
+	// Controller has changed so reset internal state
+
+	if (Controller != ActiveController || !PoseState.IsValid())
+	{
+		// Update Active Controller
+
+		ActiveController = Controller;
+
+		// Get Control Schema
+
+		const FAnimGenControlSchema& ControlSchema = ActiveController->ControlSchema;
+		const FAnimGenControlSchemaElement& ControlSchemaElement = ActiveController->ControlSchemaElement;
+		check(ControlSchema.ObservationSchema.IsValid() && ControlSchema.ObservationSchema->Schema.IsValid(ControlSchemaElement.SchemaElement));
+
+		ControlVectorSize = ControlSchema.ObservationSchema->Schema.GetObservationVectorSize(ControlSchemaElement.SchemaElement);
+		EncodedControlVectorSize = ControlSchema.ObservationSchema->Schema.GetEncodedVectorSize(ControlSchemaElement.SchemaElement);
+		ControlDistributionVectorSize = ControlSchema.ObservationSchema->Schema.GetObservationDistributionVectorSize(ControlSchemaElement.SchemaElement);
+
+		ControlVector.SetNumUninitialized({ 1, ControlVectorSize });
+		EncodedControlVector.SetNumUninitialized({ 1, EncodedControlVectorSize });
+		UE::Learning::Array::Zero(ControlVector);
+		UE::Learning::Array::Zero(EncodedControlVector);
+
+		// Get Networks
+
+		check(ActiveController->ControlEncoderNetwork);
+		check(ActiveController->LOD0Network);
+		check(ActiveController->LOD1Network);
+		check(ActiveController->LOD2Network);
+
+		ControlEncoderInference = Controller->ControlEncoderNetwork->GetNetwork()->CreateInferenceObject(1);
+		LOD0Inference = Controller->LOD0Network->GetNetwork()->CreateInferenceObject(1);
+		LOD1Inference = Controller->LOD1Network->GetNetwork()->CreateInferenceObject(1);
+		LOD2Inference = Controller->LOD2Network->GetNetwork()->CreateInferenceObject(1);
+
+		check(ActiveController->AutoEncoder);
+		EncoderInference = ActiveController->AutoEncoder->EncoderNetwork->GetNetwork()->CreateInferenceObject(1);
+		DecoderInference = ActiveController->AutoEncoder->DecoderNetwork->GetNetwork()->CreateInferenceObject(1);
+
+		// Reset Pose State
+
+		const int32 BoneNum = ActiveController->AutoEncoder->GetBoneNum();
+
+		PoseState = UAnimDatabasePoseStateLibrary::MakePoseState();
+		PoseState.PrevPoseData->RootData.Resize(1);
+		PoseState.PrevPoseData->LocalBoneData.Resize(1, BoneNum);
+		PoseState.PrevPoseData->AttributeData.Resize(1, ActiveController->AutoEncoder->GetAttributeTypes());
+		ActiveController->AutoEncoder->SetDefaultPoseData(PoseState.PrevPoseData->View());
+
+		PoseState.CurrPoseData->RootData.Resize(1);
+		PoseState.CurrPoseData->LocalBoneData.Resize(1, BoneNum);
+		PoseState.CurrPoseData->AttributeData.Resize(1, ActiveController->AutoEncoder->GetAttributeTypes());
+		ActiveController->AutoEncoder->SetDefaultPoseData(PoseState.CurrPoseData->View());
+
+		PoseState.PoseGlobalBoneData->Resize(1, ActiveController->AutoEncoder->GetBoneNum());
+
+		OutputPoseData.Resize(1, ActiveController->AutoEncoder->GetBoneNum(), ActiveController->AutoEncoder->GetAttributeTypes());
+		ActiveController->AutoEncoder->SetDefaultPoseData(OutputPoseData.View());
+
+		// Compute Pose Vector
+
+		PoseVectorSize = ActiveController->AutoEncoder->GetPoseVectorSize();
+		EncodedPoseVectorSize = ActiveController->AutoEncoder->GetEncodingSize();
+
+		PoseVector.SetNumUninitialized({ 1, PoseVectorSize });
+		UnnormalizedEncodedPoseVector.SetNumUninitialized({ 1, EncodedPoseVectorSize });
+		EncodedPoseVector.SetNumUninitialized({ 1, EncodedPoseVectorSize });
+		NextEncodedPoseVector.SetNumUninitialized({ 1, EncodedPoseVectorSize });
+		DenoiserInputVector.SetNumUninitialized({ 1, EncodedPoseVectorSize * 2 + 1 + EncodedControlVectorSize });
+		UE::Learning::Array::Zero(PoseVector);
+		UE::Learning::Array::Zero(UnnormalizedEncodedPoseVector);
+		UE::Learning::Array::Zero(EncodedPoseVector);
+		UE::Learning::Array::Zero(NextEncodedPoseVector);
+		UE::Learning::Array::Zero(DenoiserInputVector);
+
+		// Pose Needs Reset
+
+		bPoseStateRequiresReset = true;
+	}
+
+	// Check Controls are Valid
+
+	const FAnimGenControlSchema& ControlSchema = ActiveController->ControlSchema;
+	const FAnimGenControlSchemaElement& ControlSchemaElement = ActiveController->ControlSchemaElement;
+	check(ControlSchema.ObservationSchema.IsValid() && ControlSchema.ObservationSchema->Schema.IsValid(ControlSchemaElement.SchemaElement));
+
+	if (!UAnimGenControls::ValidateControlObjectMatchesSchema(ControlSchema, ControlSchemaElement, ControlObject, ControlObjectElement))
+	{
+		if (ControlObject.IsValid())
+		{
+			FWriteScopeLock ObjectLock(ControlObject.ObservationObject->Lock);
+			ControlObject.ObservationObject->Object.Reset();
+		}
+
+		bPoseStateRequiresReset = true;
+		bControlsValid = false;
+		return;
+	}
+
+	bControlsValid = true;
+
+	// Set Control Vector
+
+	{
+		FReadScopeLock ObjectLock(ControlObject.ObservationObject->Lock);
+
+		UE::Learning::Observation::SetVectorFromObject(
+			ControlVector[0],
+			ControlSchema.ObservationSchema->Schema,
+			ControlSchemaElement.SchemaElement,
+			ControlObject.ObservationObject->Object,
+			ControlObjectElement.ObjectElement);
+	}
+
+	// Reset Control Object
+
+	{
+		FWriteScopeLock ObjectLock(ControlObject.ObservationObject->Lock);
+		ControlObject.ObservationObject->Object.Reset();
+	}
+
+	// Clamp Control Vectors
+
+	UE::Learning::Observation::ClampToObservationDistributionMinMax(
+		ControlVector[0],
+		ActiveController->ControlVectorDistributionMins,
+		ActiveController->ControlVectorDistributionMaxs,
+		ControlSchema.ObservationSchema->Schema,
+		ControlSchemaElement.SchemaElement);
+
+	// Get Instance Data
+
+	check(ActiveController->IsValid());
+	check(PoseState.IsValid());
+
+	UAnimGenAutoEncoder* AutoEncoder = ActiveController->AutoEncoder;
+
+	check(ActiveController);
+	check(AutoEncoder);
+
+	const UE::PoseSearch::FPoseHistoryProvider* PoseHistoryProvider = Context.GetMessage<UE::PoseSearch::FPoseHistoryProvider>();
+	const UE::PoseSearch::IPoseHistory* PoseHistory = PoseHistoryProvider ? PoseHistoryProvider->GetPoseHistoryPtr() : nullptr;
+
+	// Reset Pose
+
+	if (bPoseStateRequiresReset)
+	{
+		// Set to default
+
+		AutoEncoder->SetDefaultPoseData(PoseState.PrevPoseData->View());
+		AutoEncoder->SetDefaultPoseData(PoseState.CurrPoseData->View());
+
+		// Reset Using Pose History
+
+		if (PoseHistory)
+		{
+			// Reset Root
+
+			FTransform CurrRootTransform = FTransform::Identity;
+
+			if (PoseHistory->GetTransformAtTime(0.0f, CurrRootTransform, Context.AnimInstanceProxy->GetSkeleton(), 0, UE::PoseSearch::WorldSpaceIndexType))
+			{
+				PoseState.CurrPoseData->RootData.RootLocations[0] = CurrRootTransform.GetLocation();
+				PoseState.CurrPoseData->RootData.RootRotations[0] = ((FQuat4f)CurrRootTransform.GetRotation()).GetNormalized();
+				PoseState.CurrPoseData->RootData.RootScales[0] = (FVector3f)CurrRootTransform.GetScale3D();
+
+				FTransform PrevRootTransform = FTransform::Identity;
+
+				if (DeltaTime > UE_SMALL_NUMBER &&
+					PoseHistory->GetTransformAtTime(-DeltaTime, PrevRootTransform, Context.AnimInstanceProxy->GetSkeleton(), 0, UE::PoseSearch::WorldSpaceIndexType))
+				{
+					const FVector3f RootLinearVelocity = (FVector3f)(CurrRootTransform.GetLocation() - PrevRootTransform.GetLocation()) / DeltaTime;
+					const FVector3f RootAngularVelocity = (FVector3f)(CurrRootTransform.GetRotation() * PrevRootTransform.GetRotation().Inverse()).GetShortestArcWith(FQuat::Identity).ToRotationVector() / DeltaTime;
+					const FVector3f RootScalarVelocity = UE::AnimDatabase::Math::VectorLogSafe(UE::AnimDatabase::Math::VectorDivMax((FVector3f)CurrRootTransform.GetScale3D(), (FVector3f)PrevRootTransform.GetScale3D())) / DeltaTime;
+
+					PoseState.PrevPoseData->RootData.RootLocations[0] = PrevRootTransform.GetLocation();
+					PoseState.PrevPoseData->RootData.RootRotations[0] = ((FQuat4f)PrevRootTransform.GetRotation()).GetNormalized();
+					PoseState.PrevPoseData->RootData.RootScales[0] = (FVector3f)PrevRootTransform.GetScale3D();
+
+					PoseState.CurrPoseData->RootData.RootLinearVelocities[0] = RootLinearVelocity;
+					PoseState.CurrPoseData->RootData.RootAngularVelocities[0] = RootAngularVelocity;
+					PoseState.CurrPoseData->RootData.RootScalarVelocities[0] = RootScalarVelocity;
+
+					PoseState.PrevPoseData->RootData.RootLinearVelocities[0] = RootLinearVelocity;
+					PoseState.PrevPoseData->RootData.RootAngularVelocities[0] = RootAngularVelocity;
+					PoseState.PrevPoseData->RootData.RootScalarVelocities[0] = RootScalarVelocity;
+				}
+				else
+				{
+					PoseState.PrevPoseData->RootData.RootLocations[0] = CurrRootTransform.GetLocation();
+					PoseState.PrevPoseData->RootData.RootRotations[0] = ((FQuat4f)CurrRootTransform.GetRotation()).GetNormalized();
+					PoseState.PrevPoseData->RootData.RootScales[0] = (FVector3f)CurrRootTransform.GetScale3D();
+				}
+			}
+			else
+			{
+				const FTransform ComponentRootTransform = Context.AnimInstanceProxy->GetComponentTransform();
+
+				PoseState.CurrPoseData->RootData.RootLocations[0] = ComponentRootTransform.GetLocation();
+				PoseState.CurrPoseData->RootData.RootRotations[0] = ((FQuat4f)ComponentRootTransform.GetRotation()).GetNormalized();
+				PoseState.CurrPoseData->RootData.RootScales[0] = (FVector3f)ComponentRootTransform.GetScale3D();
+
+				PoseState.PrevPoseData->RootData.RootLocations[0] = ComponentRootTransform.GetLocation();
+				PoseState.PrevPoseData->RootData.RootRotations[0] = ((FQuat4f)ComponentRootTransform.GetRotation()).GetNormalized();
+				PoseState.PrevPoseData->RootData.RootScales[0] = (FVector3f)ComponentRootTransform.GetScale3D();
+			}
+
+			// Reset Bones
+
+			const int32 BoneNum = AutoEncoder->GetBoneNum();
+
+			// We iterate from 1 to ignore "root" bone since that will be adjusted by root offset bone
+			for (int32 BoneIdx = 1; BoneIdx < BoneNum; BoneIdx++)
+			{
+				const int32 ParentIdx = ActiveController->AutoEncoder->GetBoneParents()[BoneIdx];
+				if (ParentIdx == INDEX_NONE) { continue; }
+
+				const int32 SkeletonBoneIdx = Context.AnimInstanceProxy->GetSkeleton()->GetReferenceSkeleton().FindBoneIndex(ActiveController->AutoEncoder->GetBoneName(BoneIdx));
+				const int32 SkeletonParentIdx = Context.AnimInstanceProxy->GetSkeleton()->GetReferenceSkeleton().FindBoneIndex(ActiveController->AutoEncoder->GetBoneName(ParentIdx));
+				if (SkeletonBoneIdx == INDEX_NONE || SkeletonParentIdx == INDEX_NONE) { continue; }
+
+				FTransform CurrBoneTransform = FTransform::Identity;
+				FTransform CurrParentTransform = FTransform::Identity;
+
+				if (PoseHistory->GetTransformAtTime(0.0f, CurrBoneTransform, Context.AnimInstanceProxy->GetSkeleton(), SkeletonBoneIdx) &&
+					PoseHistory->GetTransformAtTime(0.0f, CurrParentTransform, Context.AnimInstanceProxy->GetSkeleton(), SkeletonParentIdx))
+				{
+					FTransform CurrTransform = CurrBoneTransform * CurrParentTransform.Inverse();
+
+					PoseState.CurrPoseData->LocalBoneData.BoneLocations[0][BoneIdx] = (FVector3f)CurrTransform.GetLocation();
+					PoseState.CurrPoseData->LocalBoneData.BoneRotations[0][BoneIdx] = (FQuat4f)CurrTransform.GetRotation();
+					PoseState.CurrPoseData->LocalBoneData.BoneScales[0][BoneIdx] = (FVector3f)CurrTransform.GetScale3D();
+
+					FTransform PrevBoneTransform = FTransform::Identity;
+					FTransform PrevParentTransform = FTransform::Identity;
+
+					if (DeltaTime > UE_SMALL_NUMBER &&
+						PoseHistory->GetTransformAtTime(-DeltaTime, PrevBoneTransform, Context.AnimInstanceProxy->GetSkeleton(), SkeletonBoneIdx) &&
+						PoseHistory->GetTransformAtTime(-DeltaTime, PrevParentTransform, Context.AnimInstanceProxy->GetSkeleton(), SkeletonParentIdx))
+					{
+						FTransform PrevTransform = PrevBoneTransform * PrevParentTransform.Inverse();
+
+						const FVector3f LinearVelocity = (FVector3f)(CurrTransform.GetLocation() - PrevTransform.GetLocation()) / DeltaTime;
+						const FVector3f AngularVelocity = (FVector3f)(CurrTransform.GetRotation() * PrevTransform.GetRotation().Inverse()).GetShortestArcWith(FQuat::Identity).ToRotationVector() / DeltaTime;
+						const FVector3f ScalarVelocity = UE::AnimDatabase::Math::VectorLogSafe(UE::AnimDatabase::Math::VectorDivMax((FVector3f)CurrTransform.GetScale3D(), (FVector3f)PrevTransform.GetScale3D())) / DeltaTime;
+
+						PoseState.PrevPoseData->LocalBoneData.BoneLocations[0][BoneIdx] = (FVector3f)PrevTransform.GetLocation();
+						PoseState.PrevPoseData->LocalBoneData.BoneRotations[0][BoneIdx] = (FQuat4f)PrevTransform.GetRotation();
+						PoseState.PrevPoseData->LocalBoneData.BoneScales[0][BoneIdx] = (FVector3f)PrevTransform.GetScale3D();
+
+						PoseState.CurrPoseData->LocalBoneData.BoneLinearVelocities[0][BoneIdx] = LinearVelocity;
+						PoseState.CurrPoseData->LocalBoneData.BoneAngularVelocities[0][BoneIdx] = AngularVelocity;
+						PoseState.CurrPoseData->LocalBoneData.BoneScalarVelocities[0][BoneIdx] = ScalarVelocity;
+
+						PoseState.PrevPoseData->LocalBoneData.BoneLinearVelocities[0][BoneIdx] = LinearVelocity;
+						PoseState.PrevPoseData->LocalBoneData.BoneAngularVelocities[0][BoneIdx] = AngularVelocity;
+						PoseState.PrevPoseData->LocalBoneData.BoneScalarVelocities[0][BoneIdx] = ScalarVelocity;
+					}
+					else
+					{
+						PoseState.PrevPoseData->LocalBoneData.BoneLocations[0][BoneIdx] = (FVector3f)CurrTransform.GetLocation();
+						PoseState.PrevPoseData->LocalBoneData.BoneRotations[0][BoneIdx] = (FQuat4f)CurrTransform.GetRotation();
+						PoseState.PrevPoseData->LocalBoneData.BoneScales[0][BoneIdx] = (FVector3f)CurrTransform.GetScale3D();
+					}
+				}
+			}
+		}
+		else
+		{
+			const FTransform ComponentRootTransform = Context.AnimInstanceProxy->GetComponentTransform();
+
+			PoseState.CurrPoseData->RootData.RootLocations[0] = ComponentRootTransform.GetLocation();
+			PoseState.CurrPoseData->RootData.RootRotations[0] = ((FQuat4f)ComponentRootTransform.GetRotation()).GetNormalized();
+			PoseState.CurrPoseData->RootData.RootScales[0] = (FVector3f)ComponentRootTransform.GetScale3D();
+
+			PoseState.PrevPoseData->RootData.RootLocations[0] = ComponentRootTransform.GetLocation();
+			PoseState.PrevPoseData->RootData.RootRotations[0] = ((FQuat4f)ComponentRootTransform.GetRotation()).GetNormalized();
+			PoseState.PrevPoseData->RootData.RootScales[0] = (FVector3f)ComponentRootTransform.GetScale3D();
+		}
+
+		AutoEncoder->ToPoseVectors(PoseVector, PoseState.CurrPoseData->ConstView());
+		AutoEncoder->ClampPoseVectors(PoseVector);
+		AutoEncoder->NormalizePoseVectors(PoseVector);
+
+		// Encode Pose Vector
+
+		if (EncoderInference->GetInputSize() == PoseVector.Num<1>() &&
+			EncoderInference->GetOutputSize() == UnnormalizedEncodedPoseVector.Num<1>())
+		{
+			EncoderInference->Evaluate(UnnormalizedEncodedPoseVector, PoseVector);
+		}
+		else
+		{
+			UE::Learning::Array::Zero(UnnormalizedEncodedPoseVector);
+		}
+
+		// Normalize Encoded Pose Vector
+
+		ActiveController->NormalizeEncodedPoseVectors(EncodedPoseVector, UnnormalizedEncodedPoseVector);
+		ActiveController->ClampNormalizedPoseVectorsInplace(EncodedPoseVector);
+		UE::Learning::Array::Copy(NextEncodedPoseVector, EncodedPoseVector);
+	}
+
+	// Update Pose State and perform FK
+
+	if (bRootFromPoseHistory && PoseHistory)
+	{
+		FTransform PoseRootTransform = FTransform::Identity;
+		if (PoseHistory->GetTransformAtTime(0.0f, PoseRootTransform, Context.AnimInstanceProxy->GetSkeleton(), 0, UE::PoseSearch::WorldSpaceIndexType))
+		{
+			UE::AnimDatabase::PoseData::Teleport(
+				PoseState.CurrPoseData->RootData.View(),
+				{ PoseRootTransform.GetLocation() },
+				{ (FQuat4f)PoseRootTransform.GetRotation() });
+		}
+	}
+
+#if WITH_EDITOR
+	if (bDrawDebugPoseState)
+	{
+		UDrawDebugLibrary::DrawDebugTransform(
+			FDebugDrawer::MakeDebugDrawer(Context.AnimInstanceProxy),
+			Context.AnimInstanceProxy->GetComponentTransform(),
+			UDrawDebugLibrary::MakeDrawDebugLineStyleFromColorAndThickness(FLinearColor::White, 2.0f),
+			false,
+			10.0f);
+
+		UDrawDebugLibrary::DrawDebugRotation(
+			FDebugDrawer::MakeDebugDrawer(Context.AnimInstanceProxy),
+			PoseState.CurrPoseData->RootData.RootLocations[0],
+			((FQuat)PoseState.CurrPoseData->RootData.RootRotations[0]).GetNormalized().Rotator(),
+			UDrawDebugLibrary::MakeDrawDebugLineStyleFromColorAndThickness(FLinearColor::White, 2.0f),
+			false,
+			10.0f);
+
+		UE::AnimDatabase::PoseData::ForwardKinematics(
+			PoseState.PoseGlobalBoneData->View(),
+			PoseState.CurrPoseData->LocalBoneData.ConstView(),
+			PoseState.CurrPoseData->RootData.ConstView(), 
+			ActiveController->AutoEncoder->GetBoneParents());
+
+		UDrawDebugLibrary::DrawDebugSkeletonArrayView(
+			FDebugDrawer::MakeDebugDrawer(Context.AnimInstanceProxy),
+			PoseState.CurrPoseData->RootData.RootLocations[0],
+			PoseState.CurrPoseData->RootData.RootRotations[0],
+			PoseState.PoseGlobalBoneData->BoneLocations[0],
+			PoseState.PoseGlobalBoneData->BoneRotations[0],
+			AutoEncoder->AutoEncodedRequiredBoneIndices,
+			ActiveController->AutoEncoder->GetBoneParents());
+	}
+#endif
+
+	// Evaluate Control Encoder
+
+	if (ControlVector.Num<1>() == ControlEncoderInference->GetInputSize() &&
+		EncodedControlVector.Num<1>() == ControlEncoderInference->GetOutputSize())
+	{
+		ControlEncoderInference->Evaluate(EncodedControlVector, ControlVector);
+	}
+	else
+	{
+		UE::Learning::Array::Zero(EncodedControlVector);
+	}
+
+	// Evaluate Denoiser
+
+	const float ControllerDeltaTime = 1.0 / FMath::Max(ActiveController->FrameRate.AsDecimal(), UE_SMALL_NUMBER);
+	const float SubsampleRate = FMath::Clamp(FMath::Max(DeltaTime, ControllerDeltaTime) * ActiveController->FrameRate.AsDecimal(), 1, 8);
+
+	UE::Learning::Array::Copy(NextEncodedPoseVector, EncodedPoseVector);
+
+	if (DenoiserInputVector.Num<1>() == LOD0Inference->GetInputSize() &&
+		DenoiserInputVector.Num<1>() == LOD1Inference->GetInputSize() &&
+		DenoiserInputVector.Num<1>() == LOD2Inference->GetInputSize() &&
+		EncodedPoseVectorSize == LOD0Inference->GetOutputSize() &&
+		EncodedPoseVectorSize == LOD1Inference->GetOutputSize() &&
+		EncodedPoseVectorSize == LOD2Inference->GetOutputSize())
+	{
+		UE::Learning::Random::SampleClippedGaussianArray(NextEncodedPoseVector[0], RandomState, 0.0f, 1.0f, 3.0f);
+		ActiveController->ScaleSamplingNoiseInplace(NextEncodedPoseVector);
+
+		UE::Learning::Array::Copy(DenoiserInputVector[0].Slice(0, EncodedPoseVectorSize), EncodedPoseVector[0]);
+		UE::Learning::Array::Copy(DenoiserInputVector[0].Slice(EncodedPoseVectorSize, EncodedPoseVectorSize), NextEncodedPoseVector[0]);
+		UE::Learning::Array::Set(DenoiserInputVector[0].Slice(EncodedPoseVectorSize * 2, 1), (SubsampleRate - 1.0f) / 7.0f);
+		UE::Learning::Array::Copy(DenoiserInputVector[0].Slice(EncodedPoseVectorSize * 2 + 1, EncodedControlVectorSize), EncodedControlVector[0]);
+
+		switch (LODLevel)
+		{
+		default:
+		case 0: LOD0Inference->Evaluate(NextEncodedPoseVector, DenoiserInputVector); break;
+		case 1: LOD1Inference->Evaluate(NextEncodedPoseVector, DenoiserInputVector); break;
+		case 2: LOD2Inference->Evaluate(NextEncodedPoseVector, DenoiserInputVector); break;
+		}
+
+		ActiveController->ClampNormalizedPoseVectorsInplace(NextEncodedPoseVector);
+	}
+	else
+	{
+		UE::Learning::Array::Zero(NextEncodedPoseVector);
+	}
+
+	// Interpolate / Extrapolate using Delta Time Ratio
+
+	if (DeltaTime < ControllerDeltaTime)
+	{
+		UE::AnimDatabase::Math::LerpToTargetInplace(
+			EncodedPoseVector,
+			NextEncodedPoseVector,
+			DeltaTime * ActiveController->FrameRate.AsDecimal());
+	}
+	else
+	{
+		UE::Learning::Array::Copy(EncodedPoseVector, NextEncodedPoseVector);
+	}
+
+	ActiveController->ClampNormalizedPoseVectorsInplace(EncodedPoseVector);
+
+	// Evaluate Decoder
+
+	ActiveController->DenormalizeEncodedPoseVectors(UnnormalizedEncodedPoseVector, EncodedPoseVector);
+
+	if (DecoderInference->GetInputSize() == UnnormalizedEncodedPoseVector.Num<1>() &&
+		DecoderInference->GetOutputSize() == PoseVector.Num<1>())
+	{
+		DecoderInference->Evaluate(PoseVector, UnnormalizedEncodedPoseVector);
+	}
+	else
+	{
+		UE::Learning::Array::Zero(PoseVector);
+	}
+
+	// Update Previous Pose
+
+	*PoseState.PrevPoseData = *PoseState.CurrPoseData;
+
+	// Update Pose
+
+	FTransform PoseRootTransform = Context.AnimInstanceProxy->GetComponentTransform();
+
+	AutoEncoder->DenormalizePoseVectors(PoseVector);
+	AutoEncoder->ClampPoseVectors(PoseVector);
+	AutoEncoder->FromPoseVectors(OutputPoseData.View(), PoseVector,
+		{ PoseRootTransform.GetLocation() },
+		{ (FQuat4f)PoseRootTransform.GetRotation() });
+
+	PoseState.CurrPoseData->RootData = OutputPoseData.RootData;
+	PoseState.CurrPoseData->AttributeData = OutputPoseData.AttributeData;
+
+	if (bPoseStateRequiresReset)
+	{
+		PoseState.CurrPoseData->LocalBoneData = OutputPoseData.LocalBoneData;
+	}
+	else
+	{
+		UE::AnimDatabase::PoseData::DampInplaceUsingAdaptiveVelocityIntegrationMix(
+			PoseState.CurrPoseData->LocalBoneData.View(),
+			OutputPoseData.LocalBoneData.ConstView(),
+			DeltaTime,
+			MinPoseSmoothingTime,
+			MaxPoseSmoothingTime);
+	}
+
+	// Update Delta Time
+
+	PoseState.DeltaTime = DeltaTime;
+
+	// Write Anim Notifies and Anim Notify States
+
+	AnimNotifyEvents.Reset();
+
+	for (const TPair<FName, TSubclassOf<UAnimNotify>>& AttributeAnimNotify : AnimNotifyMap)
+	{
+		if (!AttributeAnimNotify.Value) { continue; }
+
+		const int32 AttributeIdx = AutoEncoder->FindAttributeIndex(AttributeAnimNotify.Key);
+		if (AttributeIdx == INDEX_NONE)
+		{
+			Context.LogMessage(EMessageSeverity::Warning, FText::Format(
+				LOCTEXT("AttributeAnimNotifyMapAttributeNotFoundWarning", "Cannot create Anim Notify from channel '{0}' - not found in AutoEncoder."),
+				FText::FromString(AttributeAnimNotify.Key.ToString())));
+			continue;
+		}
+
+		if (AutoEncoder->GetAttributeType(AttributeIdx) != EAnimDatabaseAttributeType::Event)
+		{
+			Context.LogMessage(EMessageSeverity::Warning, FText::Format(
+				LOCTEXT("AttributeAnimNotifyMapAttributeWrongTypeWarning", "Cannot create Anim Notify from channel '{0}' - channel is not Event type."),
+				FText::FromString(AttributeAnimNotify.Key.ToString())));
+			continue;
+		}
+
+		if (!PoseState.PrevPoseData->AttributeData.GetAttributeActive(0, AttributeIdx) ||
+			!PoseState.CurrPoseData->AttributeData.GetAttributeActive(0, AttributeIdx))
+		{
+			continue;
+		}
+
+		bool bPrevTimeUntilEventKnown = false;
+		bool bCurrTimeUntilEventKnown = false;
+		float PrevTimeUntilEvent = UE_MAX_FLT;
+		float CurrTimeUntilEvent = UE_MAX_FLT;
+		PoseState.PrevPoseData->AttributeData.GetEvent(bPrevTimeUntilEventKnown, PrevTimeUntilEvent, 0, AttributeIdx);
+		PoseState.CurrPoseData->AttributeData.GetEvent(bCurrTimeUntilEventKnown, CurrTimeUntilEvent, 0, AttributeIdx);
+
+		if (bPrevTimeUntilEventKnown &&
+			bCurrTimeUntilEventKnown &&
+			PrevTimeUntilEvent > 0.0f &&
+			CurrTimeUntilEvent < 0.0f &&
+			FMath::Abs(PrevTimeUntilEvent - CurrTimeUntilEvent) < 0.1f)
+		{
+			FAnimNotifyEvent& NotifyEvent = AnimNotifyEvents.AddDefaulted_GetRef();
+			NotifyEvent.Notify = AttributeAnimNotify.Value->GetDefaultObject<UAnimNotify>();
+			NotifyEvent.NotifyName = AutoEncoder->GetAttributeName(AttributeIdx);
+		}
+	}
+
+	for (const TPair<FName, TSubclassOf<UAnimNotifyState>>& AttributeAnimNotifyState : AnimNotifyStateMap)
+	{
+		if (!AttributeAnimNotifyState.Value) { continue; }
+
+		const int32 AttributeIdx = AutoEncoder->FindAttributeIndex(AttributeAnimNotifyState.Key);
+		if (AttributeIdx == INDEX_NONE)
+		{
+			Context.LogMessage(EMessageSeverity::Warning, FText::Format(
+				LOCTEXT("AttributeAnimNotifyStateMapAttributeNotFoundWarning", "Cannot create Anim Notify State from channel '{0}' - not found in AutoEncoder."),
+				FText::FromString(AttributeAnimNotifyState.Key.ToString())));
+			continue;
+		}
+
+		if (AutoEncoder->GetAttributeType(AttributeIdx) != EAnimDatabaseAttributeType::Bool)
+		{
+			Context.LogMessage(EMessageSeverity::Warning, FText::Format(
+				LOCTEXT("AttributeAnimNotifyStateMapAttributeWrongTypeWarning", "Cannot create Anim Notify State from channel '{0}' - channel is not Bool type."),
+				FText::FromString(AttributeAnimNotifyState.Key.ToString())));
+			continue;
+		}
+
+		if (!PoseState.CurrPoseData->AttributeData.GetAttributeActive(0, AttributeIdx))
+		{
+			continue;
+		}
+
+		if (PoseState.CurrPoseData->AttributeData.GetBool(0, AttributeIdx))
+		{
+			FAnimNotifyEvent& NotifyEvent = AnimNotifyEvents.AddDefaulted_GetRef();
+			NotifyEvent.NotifyStateClass = AttributeAnimNotifyState.Value->GetDefaultObject<UAnimNotifyState>();
+			NotifyEvent.NotifyName = AutoEncoder->GetAttributeName(AttributeIdx);
+		}
+	}
+
+	AnimNotifyEventReferences.Reset();
+	const int32 AnimNotifyEventNum = AnimNotifyEvents.Num();
+	for (int32 AnimNotifyEventIdx = 0; AnimNotifyEventIdx < AnimNotifyEventNum; AnimNotifyEventIdx++)
+	{
+		FAnimNotifyEventReference& EventReference = AnimNotifyEventReferences.Emplace_GetRef(&AnimNotifyEvents[AnimNotifyEventIdx], Context.GetAnimInstanceObject());
+
+		FAnimTickRecord TickRecord;
+		TickRecord.bActiveContext = true;
+		EventReference.GatherTickRecordData(TickRecord);
+	}
+
+	Context.AnimInstanceProxy->AddAnimNotifies(AnimNotifyEventReferences, 1.0f);
+
+	// Reset of pose state done
+
+	bPoseStateRequiresReset = false;
+}
+
+void FAnimNode_AnimGenController::CacheBones_AnyThread(const FAnimationCacheBonesContext& Context)
+{
+	DECLARE_SCOPE_HIERARCHICAL_COUNTER_ANIMNODE(CacheBones_AnyThread);
+
+	FAnimNode_Base::CacheBones_AnyThread(Context);
+}
+
+void FAnimNode_AnimGenController::Evaluate_AnyThread(FPoseContext& Output)
+{
+	DECLARE_SCOPE_HIERARCHICAL_COUNTER_ANIMNODE(Evaluate_AnyThread);
+
+	// Output Blank Pose
+
+	Output.ResetToRefPose();
+
+	if (!ActiveController || !PoseState.IsValid() || !bControlsValid) { return; }
+	check(ActiveController->IsValid());
+
+	// Output Compact Pose
+
+	const FBoneContainer& BoneContainer = Output.Pose.GetBoneContainer();
+	TArrayView<FTransform> CompactBones = Output.Pose.GetMutableBones();
+	const int32 CompactBoneNum = CompactBones.Num();
+
+	for (int32 BoneIdx = 0; BoneIdx < CompactBoneNum; BoneIdx++)
+	{
+		const int32 PoseBoneIdx = BoneContainer.GetSkeletonPoseIndexFromCompactPoseIndex(FCompactPoseBoneIndex(BoneIdx)).GetInt();
+		if (PoseBoneIdx != INDEX_NONE)
+		{
+			const int32 AutoEncoderBoneIdx = ActiveController->AutoEncoder->FindBoneIndex(BoneContainer.GetReferenceSkeleton().GetBoneName(PoseBoneIdx));
+			
+			if (AutoEncoderBoneIdx != INDEX_NONE)
+			{
+				CompactBones[BoneIdx] = FTransform(
+					((FQuat)PoseState.CurrPoseData->LocalBoneData.BoneRotations[0][AutoEncoderBoneIdx]).GetNormalized(),
+					(FVector)PoseState.CurrPoseData->LocalBoneData.BoneLocations[0][AutoEncoderBoneIdx],
+					(FVector)PoseState.CurrPoseData->LocalBoneData.BoneScales[0][AutoEncoderBoneIdx]);
+			}
+		}
+	}
+
+	// Output Curves
+
+	OutputCurves.Reset();
+
+	for (const TPair<FName, FName>& AttributeCurve : CurveMap)
+	{
+		const int32 AttributeIdx = ActiveController->AutoEncoder->FindAttributeIndex(AttributeCurve.Key);
+
+		if (AttributeIdx == INDEX_NONE)
+		{
+			Output.LogMessage(EMessageSeverity::Warning, FText::Format(
+				LOCTEXT("AttributeCurveMapAttributeNotFoundWarning", "Cannot create curve from attribute '{0}' - not found in AutoEncoder."),
+				FText::FromString(AttributeCurve.Key.ToString())));
+			continue;
+		}
+
+		if (ActiveController->AutoEncoder->GetAttributeType(AttributeIdx) != EAnimDatabaseAttributeType::Float)
+		{
+			Output.LogMessage(EMessageSeverity::Warning, FText::Format(
+				LOCTEXT("AttributeCurveMapAttributeWrongTypeWarning", "Cannot create curve from attribute '{0}' - attribute is not Float type."),
+				FText::FromString(AttributeCurve.Key.ToString())));
+			continue;
+		}
+
+		if (PoseState.CurrPoseData->AttributeData.GetAttributeActive(0, AttributeIdx))
+		{
+			OutputCurves.Add({ AttributeCurve.Value, PoseState.CurrPoseData->AttributeData.GetFloat(0, AttributeIdx) });
+		}
+	}
+
+	UE::Anim::FCurveUtils::BuildUnsorted(Output.Curve, OutputCurves);
+
+	// Write Root Motion Attribute
+
+	if (const UE::Anim::IAnimRootMotionProvider* RootMotionProvider = UE::Anim::IAnimRootMotionProvider::Get())
+	{
+		const FVector3f LocalRootLinearVelocity = PoseState.CurrPoseData->RootData.RootRotations[0].UnrotateVector(
+			PoseState.CurrPoseData->RootData.RootLinearVelocities[0]);
+
+		const FVector3f LocalRootAngularVelocity = PoseState.CurrPoseData->RootData.RootRotations[0].UnrotateVector(
+			PoseState.CurrPoseData->RootData.RootAngularVelocities[0]);
+
+		const FTransform RootMotionDelta = FTransform(
+			FQuat::MakeFromRotationVector(PoseState.DeltaTime * (FVector)LocalRootAngularVelocity),
+			PoseState.DeltaTime * (FVector)LocalRootLinearVelocity);
+
+		RootMotionProvider->SetRootMotion(RootMotionDelta, Output.CustomAttributes);
+	}
+}
+
+#undef LOCTEXT_NAMESPACE
\ No newline at end of file
diff --git a/Engine/Plugins/Experimental/Animation/AnimGen/Source/AnimGen/Private/RigVMFunction_AnimGen.cpp b/Engine/Plugins/Experimental/Animation/AnimGen/Source/AnimGen/Private/RigVMFunction_AnimGen.cpp
new file mode 100644
--- /dev/null
+++ b/Engine/Plugins/Experimental/Animation/AnimGen/Source/AnimGen/Private/RigVMFunction_AnimGen.cpp
@@ -0,0 +1,49 @@
+// Copyright Epic Games, Inc. All Rights Reserved.
+
+#include "RigVMFunction_AnimGen.h"
+
+#include "AnimGenBehavior.h"
+
+FRigVMFunction_MakeControlObject_Execute()
+{
+	Result = UAnimGenControls::MakeControlObject();
+}
+
+FRigVMFunction_MakeTrajectoryControl_Execute()
+{
+	Result = UAnimGenControls::MakeTrajectoryControlFromTransformTrajectory(Object, TransformTrajectory, RelativeTransform, SampleNum, FutureTime, ForwardVector, Tag);
+}
+
+FRigVMFunction_MakeStructControl_Execute()
+{
+	TArray<FName, TInlineAllocator<32>> TempNames;
+	TArray<FAnimGenControlObjectElement, TInlineAllocator<32>> TempElements;
+	const int32 ElementNum = Elements.Num();
+
+	TempNames.SetNumUninitialized(ElementNum);
+	TempElements.SetNumUninitialized(ElementNum);
+	for (int32 ElementIdx = 0; ElementIdx < ElementNum; ElementIdx++)
+	{
+		TempNames[ElementIdx] = Elements[ElementIdx].Name;
+		TempElements[ElementIdx] = Elements[ElementIdx].Element;
+	}
+
+	Result = UAnimGenControls::MakeStructControlFromArrayViews(Object, TempNames, TempElements, Tag);
+}
+
+FRigVMFunction_MakeNamedInclusiveDiscreteControl_Execute()
+{
+	Result = UAnimGenControls::MakeNamedInclusiveDiscreteControlFromArrayView(Object, Elements, Tag);
+}
+
+FRigVMFunction_MakeTrajectoryFollowBehaviorControl_Execute()
+{
+	Result = UAnimGenBehaviorLibrary::MakeTrajectoryFollowBehaviorControlFromArrayView(
+		Object,
+		TransformTrajectory,
+		RootBoneTransform,
+		DesiredTags,
+		TrajectorySampleNum,
+		TrajectoryFutureTime,
+		TrajectoryForwardVector);
+}
