AnimGen: Move to Experimental

diff --git a/Engine/Plugins/Experimental/Animation/AnimGen/Source/AnimGen/Public/AnimGenBehavior.h b/Engine/Plugins/Experimental/Animation/AnimGen/Source/AnimGen/Public/AnimGenBehavior.h
new file mode 100644
--- /dev/null
+++ b/Engine/Plugins/Experimental/Animation/AnimGen/Source/AnimGen/Public/AnimGenBehavior.h
@@ -0,0 +1,1280 @@
+// Copyright Epic Games, Inc. All Rights Reserved.
+
+#pragma once
+
+#include "AnimGenControl.h"
+
+#include "Animation/TrajectoryTypes.h"
+#include "InputCoreTypes.h"
+
+#include "AnimGenBehavior.generated.h"
+
+#define UE_API ANIMGEN_API
+
+class UAnimGenController;
+class UAnimGenAutoEncoder;
+
+/**
+ * A Behavior represents a specific control structure that can be used to control a character (for example - a trajectory to follow, a target point 
+ * to move towards).
+ * 
+ * To define a behavior it is required to inherit from this class and to provide the "SpecifyControl", and "MakeControlSets"
+ * functions. The "DrawDebugControl" function may also be provided to do debug drawing in the training viewport.
+ */
+UCLASS(Abstract, HideDropdown, EditInlineNew, DefaultToInstanced, BlueprintType, Blueprintable)
+class UAnimGenBehavior : public UObject
+{
+	GENERATED_BODY()
+
+#if WITH_EDITOR
+
+public:
+
+	/** Specify the control structure used by this behavior. */
+	UFUNCTION(BlueprintNativeEvent, BlueprintPure, Category = "AnimGen", meta = (ForceAsFunction))
+	UE_API FAnimGenControlSchemaElement SpecifyControl(UPARAM(ref) FAnimGenControlSchema& InControlSchema) const;
+
+	/** Make the control sets used for training from the given database and frame ranges in the database */
+	UFUNCTION(BlueprintNativeEvent, BlueprintPure, Category = "AnimGen", meta = (ForceAsFunction))
+	UE_API void MakeControlSets(TArray<FAnimGenControlSet>& OutControlSets, UPARAM(ref) FAnimGenControlObject& InControlObject, const UAnimDatabase* InDatabase, const FAnimDatabaseFrameRanges& InFrameRanges) const;
+
+	/** Draw debug the current frame in the training data */
+	UFUNCTION(BlueprintNativeEvent, BlueprintCallable, Category = "AnimGen", meta = (ForceAsFunction))
+	UE_API void DrawDebugControl(
+		const FDebugDrawer& Drawer,
+		const FAnimGenControlObject& InControlObject,
+		const FAnimGenControlObjectElement InControlObjectElement,
+		const FAnimDatabasePoseState& InPoseState,
+		const UAnimDatabase* InDatabase,
+		const UAnimGenAutoEncoder* InAutoEncoder,
+		const int32 SequenceIdx,
+		const float SequenceTime,
+		const int32 RangeStart,
+		const int32 RangeLength,
+		const FLinearColor& IdentifierColor) const;
+
+#endif
+};
+
+/** Struct to write out when resetting the preview scene of a behavior */
+USTRUCT(BlueprintType)
+struct FAnimGenBehaviorPreviewReset
+{
+	GENERATED_BODY()
+
+	/** Sequence index to reset to */
+	UPROPERTY(EditAnywhere, BlueprintReadWrite, Category = "AnimGen")
+	int32 ResetSequenceIdx = 0;
+
+	/** Time in the given sequence to reset to */
+	UPROPERTY(EditAnywhere, BlueprintReadWrite, Category = "AnimGen")
+	float ResetSequenceTime = 0.0f;
+
+	/** Root location to reset to */
+	UPROPERTY(EditAnywhere, BlueprintReadWrite, Category = "AnimGen")
+	FVector ResetLocation = FVector::ZeroVector;
+
+	/** Root rotation to reset to */
+	UPROPERTY(EditAnywhere, BlueprintReadWrite, Category = "AnimGen")
+	FRotator ResetRotation = FRotator::ZeroRotator;
+};
+
+
+/**
+ * This class can be implemented to perform preview scene logic when testing an AnimGenController
+ */
+UCLASS(Abstract, HideDropdown, EditInlineNew, DefaultToInstanced, BlueprintType, Blueprintable)
+class UAnimGenBehaviorPreview : public UObject
+{
+	GENERATED_BODY()
+
+#if WITH_EDITOR
+
+public:
+
+	/** Make the controls for the given state of this instance of the behavior */
+	UFUNCTION(BlueprintNativeEvent, BlueprintPure, Category = "AnimGen", meta = (ForceAsFunction))
+	UE_API FAnimGenControlObjectElement MakeControl(UPARAM(ref) FAnimGenControlObject& InControlObject, const FAnimDatabasePoseState& InPoseState) const;
+
+public:
+
+	/** Initialize the preview setup with the given Controller and AutoEncoder */
+	UFUNCTION(BlueprintNativeEvent, BlueprintCallable, Category = "AnimGen")
+	UE_API void OnPreviewInitialize(const UAnimGenController* Controller, const UAnimGenAutoEncoder* AutoEncoder);
+
+	/** Returns true if the preview scene requires reset. Provides the reset data OutResetData if required */
+	UFUNCTION(BlueprintNativeEvent, BlueprintCallable, Category = "AnimGen Preview Scene")
+	UE_API bool OnPreviewCheckReset(FAnimGenBehaviorPreviewReset& OutResetData, const FAnimDatabasePoseState& InPoseState);
+
+	/** Event called when the preview scene is reset */
+	UFUNCTION(BlueprintNativeEvent, BlueprintCallable, Category = "AnimGen Preview Scene")
+	UE_API void OnPreviewReset(const FAnimDatabasePoseState& InPoseState);
+
+	/** Event called when the preview scene is update each frame */
+	UFUNCTION(BlueprintNativeEvent, BlueprintCallable, Category = "AnimGen Preview Scene")
+	UE_API void OnPreviewUpdate(const FAnimDatabasePoseState& InPoseState, const float DeltaTime);
+
+	/** Event called each frame where debug draw logic can be added to the preview scene */
+	UFUNCTION(BlueprintNativeEvent, BlueprintCallable, Category = "AnimGen Preview Scene")
+	UE_API void OnPreviewDrawDebug(
+		const FDebugDrawer& Drawer,
+		const FAnimGenControlObject& InControlObject,
+		const FAnimGenControlObjectElement InControlObjectElement,
+		const FAnimDatabasePoseState& InPoseState,
+		const FLinearColor& IdentifierColor);
+
+	/** Event called when the left stick is moved in the preview scene */
+	UFUNCTION(BlueprintNativeEvent, BlueprintCallable, Category = "AnimGen Preview Scene")
+	UE_API void OnPreviewInputLeftStickAxis(const FVector Axis);
+
+	/** Event called when the right stick is moved in the preview scene */
+	UFUNCTION(BlueprintNativeEvent, BlueprintCallable, Category = "AnimGen Preview Scene")
+	UE_API void OnPreviewInputRightStickAxis(const FVector Axis);
+
+	/** Event called when the left trigger is moved in the preview scene */
+	UFUNCTION(BlueprintNativeEvent, BlueprintCallable, Category = "AnimGen Preview Scene")
+	UE_API void OnPreviewInputLeftTriggerAxis(const float Trigger);
+
+	/** Event called when the right trigger is moved in the preview scene */
+	UFUNCTION(BlueprintNativeEvent, BlueprintCallable, Category = "AnimGen Preview Scene")
+	UE_API void OnPreviewInputRightTriggerAxis(const float Trigger);
+
+	/** Event called when the mouse is moved in the preview scene */
+	UFUNCTION(BlueprintNativeEvent, BlueprintCallable, Category = "AnimGen Preview Scene")
+	UE_API void OnPreviewInputMouseMoved(const FVector MouseGroundLocation);
+
+	/** Event called when the mouse left button is pressed in the preview scene */
+	UFUNCTION(BlueprintNativeEvent, BlueprintCallable, Category = "AnimGen Preview Scene")
+	UE_API void OnPreviewInputLeftMousePressed(const FVector MouseGroundLocation);
+
+	/** Event called when the mouse left button is released in the preview scene */
+	UFUNCTION(BlueprintNativeEvent, BlueprintCallable, Category = "AnimGen Preview Scene")
+	UE_API void OnPreviewInputLeftMouseReleased(const FVector MouseGroundLocation);
+
+	/** Event called when the mouse right button is pressed in the preview scene */
+	UFUNCTION(BlueprintNativeEvent, BlueprintCallable, Category = "AnimGen Preview Scene")
+	UE_API void OnPreviewInputRightMousePressed(const FVector MouseGroundLocation);
+
+	/** Event called when the mouse right button is released in the preview scene */
+	UFUNCTION(BlueprintNativeEvent, BlueprintCallable, Category = "AnimGen Preview Scene")
+	UE_API void OnPreviewInputRightMouseReleased(const FVector MouseGroundLocation);
+
+	/** Event called when a key is pressed in the preview scene */
+	UFUNCTION(BlueprintNativeEvent, BlueprintCallable, Category = "AnimGen Preview Scene")
+	UE_API void OnPreviewInputKeyPressed(const FKey Key);
+
+	/** Event called when a key is released in the preview scene */
+	UFUNCTION(BlueprintNativeEvent, BlueprintCallable, Category = "AnimGen Preview Scene")
+	UE_API void OnPreviewInputKeyReleased(const FKey Key);
+
+public:
+
+	/** Gets the preview camera location and rotation */
+	UFUNCTION(BlueprintPure = true, Category = "AnimGen Preview Scene")
+	UE_API void GetPreviewCameraLocationAndRotation(FVector& OutCameraLocation, FRotator& OutCameraRotation) const;
+
+	/** Gets the preview camera yaw rotation */
+	UFUNCTION(BlueprintPure = true, Category = "AnimGen Preview Scene")
+	UE_API FRotator GetPreviewCameraYawRotation() const;
+
+	/** Updates the preview camera using the provided targets */
+	UFUNCTION(BlueprintPure = false, Category = "AnimGen Preview Scene")
+	UE_API void UpdatePreviewCamera(
+		const float YawVelocity,
+		const float PitchVelocity,
+		const float DistanceVelocity,
+		const float DeltaTime,
+		const bool bInvertPitch = false);
+
+	/** Updates the preview camera using the shoulder buttons */
+	UFUNCTION(BlueprintPure = false, Category = "AnimGen Preview Scene")
+	UE_API void UpdatePreviewCameraFromAxisAndShoulders(
+		const FVector Axis,
+		const bool bLeftShoulder,
+		const bool bRightShoulder,
+		const float DeltaTime,
+		const float RotationSpeed = 180.0f,
+		const float DistanceSpeed = 200.0f,
+		const bool bInvertPitch = false,
+		const float Deadzone = 0.2f,
+		const float Sensitivity = 2.0f);
+
+	/** Gets the preview scene left stick axis */
+	UFUNCTION(BlueprintPure = true, Category = "AnimGen Preview Scene")
+	UE_API FVector GetPreviewLeftStickAxis() const;
+
+	/** Gets the preview scene right stick axis */
+	UFUNCTION(BlueprintPure = true, Category = "AnimGen Preview Scene")
+	UE_API FVector GetPreviewRightStickAxis() const;
+
+	/** Gets the preview scene camera-relative left stick axis */
+	UFUNCTION(BlueprintPure = true, Category = "AnimGen Preview Scene")
+	UE_API FVector GetPreviewCameraRelativeLeftStickAxis() const;
+
+	/** Gets the preview scene camera-relative right stick axis */
+	UFUNCTION(BlueprintPure = true, Category = "AnimGen Preview Scene")
+	UE_API FVector GetPreviewCameraRelativeRightStickAxis() const;
+
+	/** Gets the preview scene left trigger axis */
+	UFUNCTION(BlueprintPure = true, Category = "AnimGen Preview Scene")
+	UE_API float GetPreviewLeftTriggerAxis() const;
+
+	/** Gets the preview scene right trigger axis */
+	UFUNCTION(BlueprintPure = true, Category = "AnimGen Preview Scene")
+	UE_API float GetPreviewRightTriggerAxis() const;
+
+	/** Gets the preview scene left trigger axis with the given sensitivity */
+	UFUNCTION(BlueprintPure = true, Category = "AnimGen Preview Scene")
+	UE_API float GetPreviewLeftTriggerAxisWithSensitivity(const float Sensitivity = 2.0f) const;
+
+	/** Gets the preview scene right trigger axis with the given sensitivity */
+	UFUNCTION(BlueprintPure = true, Category = "AnimGen Preview Scene")
+	UE_API float GetPreviewRightTriggerAxisWithSensitivity(const float Sensitivity = 2.0f) const;
+
+	/** Gets the preview scene mouse location on the ground plane */
+	UFUNCTION(BlueprintPure = true, Category = "AnimGen Preview Scene")
+	UE_API FVector GetPreviewMouseGroundLocation() const;
+
+#endif
+
+public:
+
+#if WITH_EDITORONLY_DATA
+
+	float PreviewCameraHeight = 110.0f;
+	float PreviewCameraYaw = -90.0f;
+	float PreviewCameraPitch = 30.0f;
+	float PreviewCameraDistance = 300.0f;
+	FVector PreviewCameraLocation = FVector::ZeroVector;
+	FRotator PreviewCameraRotation = FRotator::ZeroRotator;
+
+	FVector PreviewLeftStick = FVector::ZeroVector;
+	FVector PreviewRightStick = FVector::ZeroVector;
+	float PreviewLeftTrigger = 0.0f;
+	float PreviewRightTrigger = 0.0f;
+	FVector PreviewPrevLeftStick = FVector::ZeroVector;
+	FVector PreviewPrevRightStick = FVector::ZeroVector;
+	float PreviewPrevLeftTrigger = 0.0f;
+	float PreviewPrevRightTrigger = 0.0f;
+	FVector PreviewMouseGroundLocation = FVector::ZeroVector;
+	FVector PreviewPrevMouseGroundLocation = FVector::ZeroVector;
+
+#endif
+};
+
+/** Behavior for uncontrolled generation. Produces a random-walk of the dataset. */
+UCLASS(BlueprintType, Blueprintable, meta = (DisplayName = "Uncontrolled"))
+class UAnimGenBehavior_Uncontrolled : public UAnimGenBehavior
+{
+	GENERATED_BODY()
+
+#if WITH_EDITOR
+
+public:
+
+	UE_API virtual FAnimGenControlSchemaElement SpecifyControl_Implementation(FAnimGenControlSchema& InControlSchema) const override;
+
+	UE_API virtual void MakeControlSets_Implementation(TArray<FAnimGenControlSet>& OutControlSets, FAnimGenControlObject& InControlObject, const UAnimDatabase* InDatabase, const FAnimDatabaseFrameRanges& InFrameRanges) const override;
+
+#endif
+};
+
+/** Behavior for idling - user can provide some tags to control the style of the idle. */
+UCLASS(BlueprintType, Blueprintable, meta = (DisplayName = "Idle"))
+class UAnimGenBehavior_Idle : public UAnimGenBehavior
+{
+	GENERATED_BODY()
+
+#if WITH_EDITORONLY_DATA
+
+public:
+
+	/** List of tags with associated frame ranges */
+	UPROPERTY(EditAnywhere, Category = "Tags")
+	TArray<FAnimDatabaseFrameRangesEntry> TagRanges;
+
+#endif
+
+#if WITH_EDITOR
+
+public:
+
+	UE_API virtual FAnimGenControlSchemaElement SpecifyControl_Implementation(FAnimGenControlSchema& InControlSchema) const override;
+
+	UE_API virtual void MakeControlSets_Implementation(TArray<FAnimGenControlSet>& OutControlSets, FAnimGenControlObject& InControlObject, const UAnimDatabase* InDatabase, const FAnimDatabaseFrameRanges& InFrameRanges) const override;
+
+	UE_API virtual void DrawDebugControl_Implementation(
+		const FDebugDrawer& Drawer,
+		const FAnimGenControlObject& InControlObject,
+		const FAnimGenControlObjectElement InControlObjectElement,
+		const FAnimDatabasePoseState& InPoseState,
+		const UAnimDatabase* InDatabase,
+		const UAnimGenAutoEncoder* InAutoEncoder,
+		const int32 SequenceIdx,
+		const float SequenceTime,
+		const int32 RangeStart,
+		const int32 RangeLength,
+		const FLinearColor& IdentifierColor) const override;
+
+#endif
+};
+
+/** Behavior for following a trajectory which can either be generated from Mover/CharacterMovementComponent or a path to follow */
+UCLASS(BlueprintType, Blueprintable, meta = (DisplayName = "Trajectory Follow"))
+class UAnimGenBehavior_TrajectoryFollow : public UAnimGenBehavior
+{
+	GENERATED_BODY()
+
+#if WITH_EDITORONLY_DATA
+
+public:
+
+	/** Future lookahead time for the trajectory */
+	UPROPERTY(EditAnywhere, Category = "Trajectory")
+	float TrajectoryFutureTime = 1.0f;
+
+	/** Number of samples to use for the trajectory */
+	UPROPERTY(EditAnywhere, Category = "Trajectory")
+	int32 TrajectorySampleNum = 4;
+
+	/** Local root forward vector for the trajectory */
+	UPROPERTY(EditAnywhere, Category = "Trajectory")
+	FVector TrajectoryForwardVector = FVector::RightVector;
+
+	/** List of tags with associated frame ranges */
+	UPROPERTY(EditAnywhere, Category = "Tags")
+	TArray<FAnimDatabaseFrameRangesEntry> TagRanges;
+
+#endif
+
+#if WITH_EDITOR
+
+public:
+
+	UE_API virtual FAnimGenControlSchemaElement SpecifyControl_Implementation(FAnimGenControlSchema& InControlSchema) const override;
+
+	UE_API virtual void MakeControlSets_Implementation(TArray<FAnimGenControlSet>& OutControlSets, FAnimGenControlObject& InControlObject, const UAnimDatabase* InDatabase, const FAnimDatabaseFrameRanges& InFrameRanges) const override;
+
+	UE_API virtual void DrawDebugControl_Implementation(
+		const FDebugDrawer& Drawer,
+		const FAnimGenControlObject& InControlObject,
+		const FAnimGenControlObjectElement InControlObjectElement,
+		const FAnimDatabasePoseState& InPoseState,
+		const UAnimDatabase* InDatabase,
+		const UAnimGenAutoEncoder* InAutoEncoder,
+		const int32 SequenceIdx,
+		const float SequenceTime,
+		const int32 RangeStart,
+		const int32 RangeLength,
+		const FLinearColor& IdentifierColor) const override;
+
+#endif
+};
+
+/** Behavior for following a velocity and angular velocity. More useful for offline generation. */
+UCLASS(BlueprintType, Blueprintable, meta = (DisplayName = "Velocity Follow"))
+class UAnimGenBehavior_VelocityFollow : public UAnimGenBehavior
+{
+	GENERATED_BODY()
+
+#if WITH_EDITORONLY_DATA
+
+public:
+
+	/** List of tags with associated frame ranges */
+	UPROPERTY(EditAnywhere, Category = "Tags")
+	TArray<FAnimDatabaseFrameRangesEntry> TagRanges;
+
+#endif
+
+public:
+
+#if WITH_EDITOR
+
+	UE_API virtual FAnimGenControlSchemaElement SpecifyControl_Implementation(FAnimGenControlSchema& InControlSchema) const override;
+
+	UE_API virtual void MakeControlSets_Implementation(TArray<FAnimGenControlSet>& OutControlSets, FAnimGenControlObject& InControlObject, const UAnimDatabase* InDatabase, const FAnimDatabaseFrameRanges& InFrameRanges) const override;
+
+	UE_API virtual void DrawDebugControl_Implementation(
+		const FDebugDrawer& Drawer,
+		const FAnimGenControlObject& InControlObject,
+		const FAnimGenControlObjectElement InControlObjectElement,
+		const FAnimDatabasePoseState& InPoseState,
+		const UAnimDatabase* InDatabase,
+		const UAnimGenAutoEncoder* InAutoEncoder,
+		const int32 SequenceIdx,
+		const float SequenceTime,
+		const int32 RangeStart,
+		const int32 RangeLength,
+		const FLinearColor& IdentifierColor) const override;
+
+#endif
+};
+
+/** Behavior for moving toward a target with an optional facing direction */
+UCLASS(BlueprintType, Blueprintable, meta = (DisplayName = "Move to Target"))
+class UAnimGenBehavior_MoveToTarget : public UAnimGenBehavior
+{
+	GENERATED_BODY()
+
+#if WITH_EDITORONLY_DATA
+
+public:
+
+	/** List of tags with associated frame ranges */
+	UPROPERTY(EditAnywhere, Category = "Tags")
+	TArray<FAnimDatabaseFrameRangesEntry> TagRanges;
+
+#endif
+
+#if WITH_EDITOR
+
+public:
+
+	UE_API virtual FAnimGenControlSchemaElement SpecifyControl_Implementation(FAnimGenControlSchema& InControlSchema) const override;
+
+	UE_API virtual void MakeControlSets_Implementation(TArray<FAnimGenControlSet>& OutControlSets, FAnimGenControlObject& InControlObject, const UAnimDatabase* InDatabase, const FAnimDatabaseFrameRanges& InFrameRanges) const override;
+
+	UE_API virtual void DrawDebugControl_Implementation(
+		const FDebugDrawer& Drawer,
+		const FAnimGenControlObject& InControlObject,
+		const FAnimGenControlObjectElement InControlObjectElement,
+		const FAnimDatabasePoseState& InPoseState,
+		const UAnimDatabase* InDatabase,
+		const UAnimGenAutoEncoder* InAutoEncoder,
+		const int32 SequenceIdx,
+		const float SequenceTime,
+		const int32 RangeStart,
+		const int32 RangeLength,
+		const FLinearColor& IdentifierColor) const override;
+
+#endif
+};
+
+/**
+ * Behavior for following a trajectory while also interacting with an object or prop. This is useful for networked prop interactions where the 
+ * trajectory may be coming from a montage .
+ */
+UCLASS(BlueprintType, Blueprintable, meta = (DisplayName = "Trajectory Interaction"))
+class UAnimGenBehavior_TrajectoryInteraction : public UAnimGenBehavior
+{
+	GENERATED_BODY()
+
+#if WITH_EDITORONLY_DATA
+
+public:
+
+	/** Future lookahead time for the trajectory */
+	UPROPERTY(EditAnywhere, Category = "Trajectory")
+	float TrajectoryFutureTime = 1.0f;
+
+	/** Number of samples to use for the trajectory */
+	UPROPERTY(EditAnywhere, Category = "Trajectory")
+	int32 TrajectorySampleNum = 4;
+
+	/** Local root forward vector for the trajectory */
+	UPROPERTY(EditAnywhere, Category = "Trajectory")
+	FVector TrajectoryForwardVector = FVector::RightVector;
+
+	/** List of tags with associated frame ranges */
+	UPROPERTY(EditAnywhere, Category = "Tags")
+	TArray<FAnimDatabaseFrameRangesEntry> TagRanges;
+
+	/** Name of bone used for interaction */
+	UPROPERTY(EditAnywhere, Category = "Interaction")
+	FName InteractionBoneName = TEXT("attach");
+
+	/** Ranges where interaction is valid */
+	UPROPERTY(EditAnywhere, Instanced, Category = "Interaction")
+	TObjectPtr<UAnimDatabaseFrameRangesFunction> InteractionRanges;
+
+	/** Forward Vector for the interaction facing direction */
+	UPROPERTY(EditAnywhere, Category = "Interaction")
+	FVector InteractionForwardVector = FVector::ForwardVector;
+
+	/** If to use a separate forward vector for mirrored data */
+	UPROPERTY(EditAnywhere, Category = "Interaction")
+	bool bUseMirroredInteractionForwardVector = false;
+
+	/** Separate Forward Vector for the interaction facing direction when mirrored */
+	UPROPERTY(EditAnywhere, Category = "Interaction", meta = (EditCondition = "bUseMirroredInteractionForwardVector", HideEditConditionToggle))
+	FVector InteractionMirroredForwardVector = FVector(-1, 0, 0);
+
+#endif
+
+#if WITH_EDITOR
+
+public:
+
+	UE_API virtual FAnimGenControlSchemaElement SpecifyControl_Implementation(FAnimGenControlSchema& InControlSchema) const override;
+
+	UE_API virtual void MakeControlSets_Implementation(TArray<FAnimGenControlSet>& OutControlSets, FAnimGenControlObject& InControlObject, const UAnimDatabase* InDatabase, const FAnimDatabaseFrameRanges& InFrameRanges) const override;
+
+	UE_API virtual void DrawDebugControl_Implementation(
+		const FDebugDrawer& Drawer,
+		const FAnimGenControlObject& InControlObject,
+		const FAnimGenControlObjectElement InControlObjectElement,
+		const FAnimDatabasePoseState& InPoseState,
+		const UAnimDatabase* InDatabase,
+		const UAnimGenAutoEncoder* InAutoEncoder,
+		const int32 SequenceIdx,
+		const float SequenceTime,
+		const int32 RangeStart,
+		const int32 RangeLength,
+		const FLinearColor& IdentifierColor) const override;
+
+#endif
+};
+
+/** Simple Trajectory type that stores all variables required for simple spring-based prediction. */
+USTRUCT(BlueprintType)
+struct FAnimGenTrajectory
+{
+	GENERATED_BODY()
+
+public:
+
+	UPROPERTY(VisibleAnywhere, Category = "Trajectory")
+	TArray<FVector> Locations;
+
+	UPROPERTY(VisibleAnywhere, Category = "Trajectory")
+	TArray<FVector> Velocities;
+
+	UPROPERTY(VisibleAnywhere, Category = "Trajectory")
+	TArray<FVector> Accelerations;
+
+	UPROPERTY(VisibleAnywhere, Category = "Trajectory")
+	TArray<FRotator> Rotations;
+
+	UPROPERTY(VisibleAnywhere, Category = "Trajectory")
+	TArray<FVector> AngularVelocities;
+
+	UPROPERTY(VisibleAnywhere, Category = "Trajectory")
+	TArray<FVector> Directions;
+};
+
+/** Blueprint library of helper functions for making and testing behaviors including functions for preparing, processing, and making controls */
+UCLASS(BlueprintType, meta = (BlueprintThreadSafe))
+class UAnimGenBehaviorLibrary : public UBlueprintFunctionLibrary
+{
+	GENERATED_BODY()
+
+public:
+
+	/** Make a control object element for the Uncontrolled Behavior */
+	UFUNCTION(BlueprintPure, Category = "AnimGen")
+	static UE_API FAnimGenControlObjectElement MakeUncontrolledBehaviorControl(UPARAM(ref) FAnimGenControlObject& InControlObject);
+
+	/** Make a control object element for the Idle Behavior */
+	UFUNCTION(BlueprintPure, Category = "AnimGen", meta = (AutoCreateRefTerm = "DesiredTags"))
+	static UE_API FAnimGenControlObjectElement MakeIdleBehaviorControl(UPARAM(ref) FAnimGenControlObject& InControlObject, const TArray<FName>& DesiredTags);
+	static UE_API FAnimGenControlObjectElement MakeIdleBehaviorControlFromArrayView(FAnimGenControlObject& InControlObject, const TArrayView<const FName> DesiredTags = {});
+
+	/** Make a control object element for the Trajectory Following Behavior */
+	UFUNCTION(BlueprintPure, Category = "AnimGen", meta = (AutoCreateRefTerm = "TransformTrajectory, RootBoneTransform, DesiredTags"))
+	static UE_API FAnimGenControlObjectElement MakeTrajectoryFollowBehaviorControl(
+		UPARAM(ref) FAnimGenControlObject& InControlObject,
+		const FTransformTrajectory& TransformTrajectory,
+		const FTransform& RootBoneTransform,
+		const TArray<FName>& DesiredTags,
+		const int32 InTrajectorySampleNum = 4,
+		const float InTrajectoryFutureTime = 1.0f,
+		const FVector InTrajectoryForwardVector = FVector::RightVector);
+
+	static UE_API FAnimGenControlObjectElement MakeTrajectoryFollowBehaviorControlFromArrayView(
+		FAnimGenControlObject& InControlObject,
+		const FTransformTrajectory& TransformTrajectory,
+		const FTransform& RootBoneTransform,
+		const TArrayView<const FName> DesiredTags = {},
+		const int32 InTrajectorySampleNum = 4,
+		const float InTrajectoryFutureTime = 1.0f,
+		const FVector InTrajectoryForwardVector = FVector::RightVector);
+
+	/** Make a control object element for the Move To Target Behavior */
+	UFUNCTION(BlueprintPure, Category = "AnimGen", meta = (AutoCreateRefTerm = "RootBoneTransform, DesiredTags"))
+	static UE_API FAnimGenControlObjectElement MakeMoveToTargetBehaviorControl(
+		UPARAM(ref) FAnimGenControlObject& InControlObject,
+		const FVector TargetLocation,
+		const FVector TargetDirection,
+		const bool bUseTargetDirection,
+		const FTransform& RootBoneTransform,
+		const TArray<FName>& DesiredTags);
+
+	static UE_API FAnimGenControlObjectElement MakeMoveToTargetBehaviorControlFromArrayView(
+		FAnimGenControlObject& InControlObject,
+		const FVector TargetLocation,
+		const FVector TargetDirection,
+		const bool bUseTargetDirection,
+		const FTransform& RootBoneTransform,
+		const TArrayView<const FName> DesiredTags = {});
+
+	/** Make a control object element for the Trajectory Interaction Behavior */
+	UFUNCTION(BlueprintPure, Category = "AnimGen", meta = (AutoCreateRefTerm = "TransformTrajectory, RootBoneTransform, DesiredTags"))
+	static UE_API FAnimGenControlObjectElement MakeTrajectoryInteractionBehaviorControl(
+		UPARAM(ref) FAnimGenControlObject& InControlObject,
+		const FTransformTrajectory& TransformTrajectory,
+		const FTransform& RootBoneTransform,
+		const TArray<FName>& DesiredTags,
+		const FVector InteractionLocation = FVector::ZeroVector,
+		const FVector InteractionDirection = FVector::ForwardVector,
+		const bool bInteractionActive = false,
+		const int32 InTrajectorySampleNum = 4,
+		const float InTrajectoryFutureTime = 1.0f,
+		const FVector InTrajectoryForwardVector = FVector::RightVector);
+
+	static UE_API FAnimGenControlObjectElement MakeTrajectoryInteractionBehaviorControlFromArrayView(
+		FAnimGenControlObject& InControlObject,
+		const FTransformTrajectory& TransformTrajectory,
+		const FTransform& RootBoneTransform,
+		const TArrayView<const FName> DesiredTags = {},
+		const FVector InteractionLocation = FVector::ZeroVector,
+		const FVector InteractionDirection = FVector::ForwardVector,
+		const bool bInteractionActive = false,
+		const int32 InTrajectorySampleNum = 4,
+		const float InTrajectoryFutureTime = 1.0f,
+		const FVector InTrajectoryForwardVector = FVector::RightVector);
+
+public:
+
+	/** Make a random bool value */
+	UFUNCTION(BlueprintPure, Category = "AnimGen")
+	static UE_API bool MakeRandomBool(UPARAM(ref) int32& State);
+
+	/** Make a random bool value with the given true probability */
+	UFUNCTION(BlueprintPure, Category = "AnimGen")
+	static UE_API bool MakeRandomBoolWithProbability(UPARAM(ref) int32& State, const float TrueProbability = 0.5f);
+
+	/** Switch on a randomly generated bool value */
+	UFUNCTION(BlueprintCallable, Category = "AnimGen", meta = (ExpandBoolAsExecs = "ReturnValue"))
+	static UE_API bool SwitchOnRandomBool(UPARAM(ref) int32& State);
+
+	/** Switch on a randomly generated bool value with the given true probability */
+	UFUNCTION(BlueprintCallable, Category = "AnimGen", meta = (ExpandBoolAsExecs = "ReturnValue"))
+	static UE_API bool SwitchOnRandomBoolWithProbability(UPARAM(ref) int32& State, const float TrueProbability = 0.5f);
+
+	/** Make a random float in the range [0, 1] */
+	UFUNCTION(BlueprintPure, Category = "AnimGen")
+	static UE_API float MakeRandomFloat(UPARAM(ref) int32& State);
+
+	/** Make a random float in the provided range */
+	UFUNCTION(BlueprintPure, Category = "AnimGen")
+	static UE_API float MakeRandomFloatInRange(UPARAM(ref) int32& State, const float Min = 0.0f, const float Max = 1.0f);
+
+	/** Randomly choose from a given set of floats */
+	UFUNCTION(BlueprintPure, Category = "AnimGen")
+	static UE_API float MakeRandomFloatChoice(UPARAM(ref) int32& State, const TArray<float>& Values);
+	static UE_API float MakeRandomFloatChoiceFromArrayView(int32& State, const TArrayView<const float> Values);
+
+	/** Sample an integer in the given range (inclusive) */
+	UFUNCTION(BlueprintPure, Category = "AnimGen")
+	static UE_API int32 MakeRandomIntegerInRangeInclusive(UPARAM(ref) int32& State, const int32 Min, const int32 Max);
+
+	/** Sample a random enum value from the given enum class */
+	UFUNCTION(BlueprintPure, Category = "AnimGen")
+	static UE_API uint8 MakeRandomEnumValue(UPARAM(ref) int32& State, const UEnum* Enum);
+
+	/** Sample a random 3d vector from a Gaussian Distribution */
+	UFUNCTION(BlueprintPure, Category = "AnimGen")
+	static UE_API FVector MakeRandomVectorFromGaussianDistribution(
+		UPARAM(ref) int32& State,
+		const FVector Mean = FVector::ZeroVector,
+		const FVector Std = FVector(1, 1, 1));
+
+	/** Sample a random 3d direction */
+	UFUNCTION(BlueprintPure, Category = "AnimGen")
+	static UE_API FVector MakeRandomDirection3D(
+		UPARAM(ref) int32& State,
+		const FVector Mean = FVector::ZeroVector,
+		const FVector Std = FVector(1, 1, 1),
+		const FVector Default = FVector::ForwardVector);
+
+	/** Sample a random 2d direction (the Z coordinate will be zero) */
+	UFUNCTION(BlueprintPure, Category = "AnimGen")
+	static UE_API FVector MakeRandomDirection2D(UPARAM(ref) int32& State);
+
+	/** Sample a random 2d location from a Gaussian Distribution (the Z coordinate will be zero) */
+	UFUNCTION(BlueprintPure, Category = "AnimGen")
+	static UE_API FVector MakeRandomLocationFromGaussian2D(UPARAM(ref) int32& State, const FVector Mean = FVector::ZeroVector, const float Std = 100.0f);
+
+	/** Sample a random sequence index and sequence time from a set of frame range */
+	UFUNCTION(BlueprintPure = false, Category = "AnimGen")
+	static UE_API void MakeRandomSequenceIndexAndTimeFromFrameRanges(
+		int32& OutSequenceIdx,
+		float& OutSequenceTime,
+		UPARAM(ref) int32& State,
+		const FAnimDatabaseFrameRanges& FrameRanges,
+		const float FrameDeltaTime);
+
+public:
+
+	/** Projects a location onto the floor plane given by the provided height */
+	UFUNCTION(BlueprintPure, Category = "AnimGen")
+	static UE_API FVector ProjectOntoFloor(const FVector Location, const float FloorHeight = 0.0f);
+
+	/** Clamps one vector to some maximum distance from another vector */
+	UFUNCTION(BlueprintPure, Category = "AnimGen")
+	static UE_API FVector ClampToMaxDistanceFromVector(
+		const FVector InVector,
+		const FVector FromVector,
+		const float MaxDistance = 100.0f);
+
+	/** Apply proper dead-zone and sensitivity treatment to stick input */
+	UFUNCTION(BlueprintPure, Category = "AnimGen")
+	static UE_API void ApplyStickDeadzoneAndSensitivity(
+		FVector& OutAxis,
+		FVector& OutDirection,
+		float& OutLength,
+		bool& bOutNotInDeadzone,
+		const FVector Axis,
+		const float Deadzone = 0.2f,
+		const float Sensitivity = 2.0f);
+
+	/** Apply sensitivity to a trigger value */
+	UFUNCTION(BlueprintPure, Category = "AnimGen")
+	static UE_API float ApplyTriggerSensitivity(const float Trigger, const float Sensitivity = 2.0f);
+
+	/** Check if the stick is not in the deadzone */
+	UFUNCTION(BlueprintPure, Category = "AnimGen")
+	static UE_API bool IsStickNotInDeadzone(const FVector Axis, const float Deadzone = 0.2f);
+
+	/** Check if the stick is in the deadzone */
+	UFUNCTION(BlueprintPure, Category = "AnimGen")
+	static UE_API bool IsStickInDeadzone(const FVector Axis, const float Deadzone = 0.2f);
+
+	/** Convert a desired direction to a desired rotation */
+	UFUNCTION(BlueprintPure, Category = "AnimGen")
+	static UE_API FRotator DesiredDirectionToRotation(const FVector Direction);
+
+	/** Convert a desired direction to a desired rotation quaternion */
+	UFUNCTION(BlueprintPure, Category = "AnimGen")
+	static UE_API FQuat DesiredDirectionToRotationQuat(const FVector Direction);
+
+	/** Update the desired velocity from a stick axis */
+	UFUNCTION(BlueprintPure = false, Category = "AnimGen")
+	static UE_API void UpdateDesiredVelocityFromAxis(
+		UPARAM(ref) FVector& DesiredVelocity,
+		const FVector Axis,
+		const float VelocitySpeed = 150.0f,
+		const float Deadzone = 0.2f,
+		const float Sensitivity = 2.0f);
+
+	/** Update the desired rotation from a stick axis */
+	UFUNCTION(BlueprintPure = false, Category = "AnimGen")
+	static UE_API void UpdateDesiredRotationFromAxis(
+		UPARAM(ref) FRotator& DesiredRotation,
+		const FVector Axis,
+		const float Deadzone = 0.2f,
+		const float Sensitivity = 2.0f);
+
+	/** Update the desired velocity and rotation from two stick axes */
+	UFUNCTION(BlueprintPure = false, Category = "AnimGen")
+	static UE_API void UpdateDesiredVelocityAndRotationFromAxes(
+		UPARAM(ref) FVector& DesiredVelocity,
+		UPARAM(ref) FRotator& DesiredRotation,
+		const FVector VelocityAxis,
+		const FVector FacingAxis,
+		const float VelocitySpeed = 150.0f,
+		const float Deadzone = 0.2f,
+		const float Sensitivity = 2.0f);
+
+	/** Update the desired velocity and rotation from a stick axes and camera trigger */
+	UFUNCTION(BlueprintPure = false, Category = "AnimGen")
+	static UE_API void UpdateDesiredVelocityAndRotationFromAxisAndCameraTrigger(
+		UPARAM(ref) FVector& DesiredVelocity,
+		UPARAM(ref) FRotator& DesiredRotation,
+		const FRotator CameraYawRotation,
+		const FVector VelocityAxis,
+		const float FacingTrigger,
+		const float MinimumSpeed = 75.0f,
+		const float MaximumSpeed = 150.0f,
+		const float Deadzone = 0.2f,
+		const float Sensitivity = 2.0f);
+
+	/** Find the index of the most similar value in an array of floats */
+	UFUNCTION(BlueprintPure = false, Category = "AnimGen")
+	static UE_API int32 FindNearestIndexInFloatArray(const TArray<float>& Values, const float Value);
+
+	/** Cycle through an array of names and return the current name */
+	UFUNCTION(BlueprintPure = false, Category = "AnimGen")
+	static UE_API FName CycleNameInArray(const TArray<FName>& Names, const FName CurrentName, const int32 Steps = 1);
+	static UE_API FName CycleNameInArrayView(const TArrayView<const FName> Names, const FName CurrentName, const int32 Steps = 1);
+
+	/** Check if two locations are within some distance threshold */
+	UFUNCTION(BlueprintPure, Category = "AnimGen")
+	static UE_API bool LocationDifferenceBelowThreshold(
+		const FVector CurrentLocation,
+		const FVector TargetLocation,
+		const float LocationThreshold = 25.0f);
+
+	/** Check if locations and directions are within some distance and angle threshold (in degrees) */
+	UFUNCTION(BlueprintPure, Category = "AnimGen")
+	static UE_API bool LocationAndDirectionDifferenceBelowThreshold(
+		const FVector CurrentLocation,
+		const FVector CurrentDirection,
+		const FVector TargetLocation,
+		const FVector TargetDirection,
+		const float LocationThreshold = 25.0f,
+		const float DirectionThresoldAngle = 10.0f);
+
+public:
+
+	/** Get the velocity along the path at the given time */
+	UFUNCTION(BlueprintPure = true, Category = "AnimGen")
+	static UE_API FVector PathVelocityAtTime(const TArray<FVector>& Points, const float Time);
+	static UE_API FVector PathVelocityAtTimeFromArrayView(const TArrayView<const FVector> Points, const float Time);
+
+	/** Get the location along the path at the given time */
+	UFUNCTION(BlueprintPure = true, Category = "AnimGen")
+	static UE_API FVector PathLocationAtTime(const TArray<FVector>& Points, const float Time);
+	static UE_API FVector PathLocationAtTimeFromArrayView(const TArrayView<const FVector> Points, const float Time);
+
+	/** Get the direction along the path at the given time */
+	UFUNCTION(BlueprintPure = true, Category = "AnimGen")
+	static UE_API FVector PathDirectionAtTime(const TArray<FVector>& Points, const float Time);
+	static UE_API FVector PathDirectionAtTimeFromArrayView(const TArrayView<const FVector> Points, const float Time);
+
+	/** Get the rotation along the path at the given time */
+	UFUNCTION(BlueprintPure = true, Category = "AnimGen")
+	static UE_API FRotator PathRotationAtTime(const TArray<FVector>& Points, const float Time);
+	static UE_API FRotator PathRotationAtTimeFromArrayView(const TArrayView<const FVector> Points, const float Time);
+
+	/** Adds a point to a path if it is further away than the given minimum distance from the last point added */
+	UFUNCTION(BlueprintPure = false, Category = "AnimGen", meta = (ReturnDisplayName = "Added"))
+	static UE_API bool PathAddPoint(UPARAM(ref) TArray<FVector>& Points, const FVector NewPoint, const float MinDistance = 50.0f);
+
+	/** Get the total distance of a path */
+	UFUNCTION(BlueprintPure = true, Category = "AnimGen")
+	static UE_API float PathTotalDistance(const TArray<FVector>& Points);
+	static UE_API float PathTotalDistanceFromArrayView(const TArrayView<const FVector> Points);
+
+public:
+
+	/** Make a transform trajectory from an array of locations, rotations, and times */
+	UFUNCTION(BlueprintPure = true, Category = "AnimGen")
+	static UE_API void MakeTransformTrajectory(FTransformTrajectory& OutTransformTrajectory, const TArray<FVector>& Locations, const TArray<FRotator>& Rotations, const TArray<float>& Times);
+	static UE_API void MakeTransformTrajectoryFromArrayViews(FTransformTrajectory& OutTransformTrajectory, const TArrayView<const FVector> Locations, const TArrayView<const FRotator> Rotations, const TArrayView<const float> Times);
+
+	/** Make a transform trajectory from an AnimGenTrajectory */
+	UFUNCTION(BlueprintPure = true, Category = "AnimGen")
+	static UE_API void MakeTransformTrajectoryFromTrajectory(FTransformTrajectory& OutTransformTrajectory, const FAnimGenTrajectory& Trajectory, const float FutureTime = 1.0f);
+
+	/** Make a transform trajectory from a path */
+	UFUNCTION(BlueprintPure = true, Category = "AnimGen")
+	static UE_API void MakeTransformTrajectoryFromPath(FTransformTrajectory& OutTransformTrajectory, const TArray<FVector>& Path, const float StartDistance = 0.0f, const float PathPointDistance = 50.0f, const float PathSpeed = 150.0f, const float FutureTime = 1.0f, const float SampleInterval = 0.01666666666f);
+	static UE_API void MakeTransformTrajectoryFromPathArrayView(FTransformTrajectory& OutTransformTrajectory, const TArrayView<const FVector> Path, const float StartDistance = 0.0f, const float PathPointDistance = 50.0f, const float PathSpeed = 150.0f, const float FutureTime = 1.0f, const float SampleInterval = 0.01666666666f);
+
+	/** Project a transform trajectory onto the ground */
+	UFUNCTION(BlueprintPure = false, Category = "AnimGen")
+	static UE_API void ProjectTransformTrajectoryOntoFloor(UPARAM(ref) FTransformTrajectory& InOutTransformTrajectory, const float FloorHeight = 0.0f);
+
+	/** Get the number of samples in an AnimGenTrajectory */
+	UFUNCTION(BlueprintPure = true, Category = "AnimGen")
+	static UE_API int32 GetTrajectorySampleNum(const FAnimGenTrajectory& Trajectory);
+
+	/** Get the locations in an AnimGenTrajectory */
+	UFUNCTION(BlueprintPure = true, Category = "AnimGen")
+	static UE_API const TArray<FVector>& GetTrajectoryLocations(const FAnimGenTrajectory& Trajectory);
+
+	/** Get the directions in an AnimGenTrajectory */
+	UFUNCTION(BlueprintPure = true, Category = "AnimGen")
+	static UE_API const TArray<FVector>& GetTrajectoryDirections(const FAnimGenTrajectory& Trajectory);
+
+	/** Reset an AnimGenTrajectory to a given state */
+	UFUNCTION(BlueprintPure = false, Category = "AnimGen")
+	static UE_API void ResetTrajectory(
+		UPARAM(ref) FAnimGenTrajectory& Trajectory,
+		const int32 SampleNum = 4,
+		const FVector InitialLocation = FVector::ZeroVector,
+		const FRotator InitialRotation = FRotator::ZeroRotator,
+		const FVector ForwardVector = FVector(0, 1, 0));
+
+	/** Update and predict the future for an AnimGenTrajectory */
+	UFUNCTION(BlueprintPure = false, Category = "AnimGen", meta = (AdvancedDisplay = 8))
+	static UE_API void UpdateAndPredictFutureTrajectory(
+		UPARAM(ref) FAnimGenTrajectory& Trajectory,
+		const FVector DesiredVelocity,
+		const FRotator DesiredRotation,
+		const float DeltaTime,
+		const float FutureTime = 1.0f,
+		const float MaximumDesiredVelocity = 500.0f,
+		const float MaximumAcceleration = 1000.0f,
+		const float MaximumAngularVelocity = 360.0f,
+		const FVector ForwardVector = FVector(0, 1, 0));
+
+	/** Teleport an AnimGenTrajectory. Updates velocities and future estimations. */
+	UFUNCTION(BlueprintPure = false, Category = "AnimGen")
+	static UE_API void TeleportTrajectory(
+		UPARAM(ref) FAnimGenTrajectory& Trajectory,
+		const FVector NewLocation,
+		const FRotator NewRotation);
+
+public:
+
+	/** Applies a simple avoidance heuristic to a movement velocity */
+	UFUNCTION(BlueprintPure = false, Category = "AnimGen")
+	static UE_API FVector ApplySimpleAvoidanceToVelocity(
+		const FVector Velocity,
+		const TArray<AActor*>& Actors,
+		const AActor* SelfActor,
+		const float MinAvoidanceRadius = 25.0f,
+		const float MaxAvoidanceRadius = 150.0f,
+		const float AvoidanceVelocity = 200.0f);
+
+	static UE_API FVector ApplySimpleAvoidanceToVelocityArrayView(
+		const FVector Velocity,
+		const TArrayView<AActor* const> Actors,
+		const AActor* SelfActor,
+		const float MinAvoidanceRadius = 25.0f,
+		const float MaxAvoidanceRadius = 150.0f,
+		const float AvoidanceVelocity = 200.0f);
+
+	/** Checks if an actor can look at another actor according to some distance and angle thresholds */
+	UFUNCTION(BlueprintPure = false, Category = "AnimGen")
+	static UE_API bool CanLookAtActor(
+		const FVector LookAtDirection,
+		const AActor* Actor,
+		const AActor* SelfActor,
+		const float DistanceThreshold = 200.0f,
+		const float DotAngleThreshold = 0.75f);
+
+	/** Finds an actor to look-at according to some distance and angle thresholds */
+	UFUNCTION(BlueprintPure = false, Category = "AnimGen")
+	static UE_API AActor* FindLookAtActor(
+		const FVector LookAtDirection,
+		const TArray<AActor*>& Actors,
+		const AActor* SelfActor,
+		const float DistanceThreshold = 200.0f,
+		const float DotAngleThreshold = 0.75f);
+
+	static UE_API AActor* FindLookAtActorArrayView(
+		const FVector LookAtDirection,
+		const TArray<AActor*>& Actors,
+		const AActor* SelfActor,
+		const float DistanceThreshold = 200.0f,
+		const float DotAngleThreshold = 0.75f);
+};
+
+/** Simple Behavior preview for the uncontrolled behavior */
+UCLASS(BlueprintType, Blueprintable, meta = (DisplayName = "Uncontrolled"))
+class UAnimGenBehaviorPreview_Uncontrolled : public UAnimGenBehaviorPreview
+{
+	GENERATED_BODY()
+
+#if WITH_EDITOR
+
+public:
+
+	virtual FAnimGenControlObjectElement MakeControl_Implementation(FAnimGenControlObject& InControlObject, const FAnimDatabasePoseState& InPoseState) const override;
+
+#endif
+};
+
+/** Simple Behavior preview for the idle behavior */
+UCLASS(BlueprintType, Blueprintable, meta = (DisplayName = "Idle"))
+class UAnimGenBehaviorPreview_Idle : public UAnimGenBehaviorPreview
+{
+	GENERATED_BODY()
+
+#if WITH_EDITORONLY_DATA
+
+public:
+
+	/** List of styles to cycle through */
+	UPROPERTY(EditAnywhere, Category = "Styles")
+	TArray<FName> StyleTags;
+
+protected:
+
+	/** Current Desired Style */
+	FName DesiredStyle = NAME_None;
+
+	/** Current Style Index */
+	int32 StyleIndex = 0;
+
+#endif
+
+#if WITH_EDITOR
+
+public:
+
+	virtual FAnimGenControlObjectElement MakeControl_Implementation(FAnimGenControlObject& InControlObject, const FAnimDatabasePoseState& InPoseState) const override;
+
+	virtual void OnPreviewInputKeyPressed_Implementation(const FKey Key) override;
+
+	virtual void OnPreviewDrawDebug_Implementation(
+		const FDebugDrawer& Drawer,
+		const FAnimGenControlObject& InControlObject,
+		const FAnimGenControlObjectElement InControlObjectElement,
+		const FAnimDatabasePoseState& InPoseState,
+		const FLinearColor& IdentifierColor) override;
+#endif
+};
+
+/** Behavior preview for Trajectory Following that allows gamepad control or drawing of trajectory on the ground */
+UCLASS(BlueprintType, Blueprintable, meta = (DisplayName = "Trajectory Follow"))
+class UAnimGenBehaviorPreview_TrajectoryFollow : public UAnimGenBehaviorPreview
+{
+	GENERATED_BODY()
+
+#if WITH_EDITORONLY_DATA
+
+public:
+
+	/** Minimum speed when walking */
+	UPROPERTY(EditAnywhere, Category = "Locomotion", meta = (ForceUnits = "cm/s"))
+	float MinimumWalkSpeed = 75.0f;
+
+	/** Maximum Speed used for walking */
+	UPROPERTY(EditAnywhere, Category = "Locomotion", meta = (ForceUnits = "cm/s"))
+	float MaximumWalkSpeed = 150.0f;
+
+	/** Minimum Speed used for running */
+	UPROPERTY(EditAnywhere, Category = "Locomotion", meta = (ForceUnits = "cm/s"))
+	float MinimumRunSpeed = 200.0f;
+
+	/** Maximum Speed used for running */
+	UPROPERTY(EditAnywhere, Category = "Locomotion", meta = (ForceUnits = "cm/s"))
+	float MaximumRunSpeed = 500.0f;
+
+	/** Smoothing to apply to speed changes */
+	UPROPERTY(EditAnywhere, Category = "Locomotion", meta = (ForceUnits = "s"))
+	float SpeedSmoothingTime = 0.1f;
+
+	/** Maximum acceleration for movement model */
+	UPROPERTY(EditAnywhere, Category = "Locomotion", meta = (ForceUnits = "cm/s^2"))
+	float MaximumAcceleration = 1000.0f;
+
+	/** Maximum angular velocity used for movement model */
+	UPROPERTY(EditAnywhere, Category = "Locomotion", meta = (ForceUnits = "deg/s"))
+	float MaximumAngularVelocity = 360.0f;
+
+	/** Tag used for standing */
+	UPROPERTY(EditAnywhere, Category = "Locomotion")
+	FName StandTag = NAME_None;
+
+	/** Tag used for walking */
+	UPROPERTY(EditAnywhere, Category = "Locomotion")
+	FName WalkTag = NAME_None;
+
+	/** Tag used for running */
+	UPROPERTY(EditAnywhere, Category = "Locomotion")
+	FName RunTag = NAME_None;
+
+public:
+
+	/** Future lookahead time for the trajectory */
+	UPROPERTY(EditAnywhere, Category = "Trajectory", meta = (ForceUnits = "s"))
+	float TrajectoryFutureTime = 1.0f;
+
+	/** Number of samples to use for the trajectory */
+	UPROPERTY(EditAnywhere, Category = "Trajectory")
+	int32 TrajectorySampleNum = 4;
+
+	/** Local root forward vector for the trajectory */
+	UPROPERTY(EditAnywhere, Category = "Trajectory")
+	FVector TrajectoryForwardVector = FVector::RightVector;
+
+public:
+
+	/** List of styles to cycle through */
+	UPROPERTY(EditAnywhere, Category = "Styles")
+	TArray<FName> StyleTags;
+
+public:
+
+	/** Maximum speed at which the camera can rotate */
+	UPROPERTY(EditAnywhere, Category = "Camera", meta = (ForceUnits = "deg/s"))
+	float CameraRotationSpeed = 180.0f;
+
+	/** Maximum speed at which the camera can move closer and further from the character */
+	UPROPERTY(EditAnywhere, Category = "Camera", meta = (ForceUnits = "cm/s"))
+	float CameraZoomSpeed = 100.0f;
+
+	/** Smoothing to apply to camera input controls */
+	UPROPERTY(EditAnywhere, Category = "Camera", meta = (ForceUnits = "s"))
+	float CameraInputSmoothingTime = 0.1f;
+
+	/** If to invert the pitch control for the camera */
+	UPROPERTY(EditAnywhere, Category = "Camera")
+	bool bCameraInvertPitch = true;
+
+public:
+
+	/** Speed to use for path following */
+	UPROPERTY(EditAnywhere, Category = "Path", meta = (ForceUnits = "cm/s"))
+	float PathFollowSpeed = 150.0f;
+
+	/** Minimum distance between points in the path */
+	UPROPERTY(EditAnywhere, Category = "Path", meta = (ForceUnits = "cm"))
+	float PathPointDistance = 50.0f;
+
+	/** If the path facing direction should face the camera */
+	UPROPERTY(EditAnywhere, Category = "Path")
+	bool bPathFacingCamera = false;
+
+protected:
+
+	/** If to follow a path rather than the trajectory */
+	bool bFollowPath = false;
+
+	/** If a path is being drawn */
+	bool bDrawPath = false;
+
+	/** Path to follow */
+	TArray<FVector> FollowPath;
+
+	/** Path being drawn */
+	TArray<FVector> DrawPath;
+
+	/** Current Distance along Path */
+	float PathDistance = 0.0f;
+
+protected:
+
+	/** Current Desired Style */
+	FName DesiredStyle = NAME_None;
+
+	/** Current Desired Gait */
+	FName DesiredGait = NAME_None;
+
+	/** If walking is currently active */
+	bool bWalkActive = false;
+
+	/** Current Style Index */
+	int32 StyleIndex = 0;
+
+	/** Current desired speed */
+	float DesiredMinimumSpeed = 0.0f;
+	float DesiredMaximumSpeed = 0.0f;
+
+	/** Current change in desired speed */
+	float DesiredMinimumSpeedVelocity = 0.0f;
+	float DesiredMaximumSpeedVelocity = 0.0f;
+
+	/** Current Desired Rotation */
+	FRotator DesiredRotation = FRotator::ZeroRotator;
+
+	/** Current Desired Velocity */
+	FVector DesiredVelocity = FVector::ZeroVector;
+
+	/** Current Trajectory */
+	FAnimGenTrajectory Trajectory;
+
+	/** Current Transform Trajectory */
+	FTransformTrajectory TransformTrajectory;
+
+protected:
+
+	/** Camera Parameters */
+
+	float CameraYawVelocity = 0.0f;
+	float CameraYawAcceleration = 0.0f;
+	float CameraPitchVelocity = 0.0f;
+	float CameraPitchAcceleration = 0.0f;
+	float CameraDistanceDelta = 0.0f;
+	float CameraDistanceVelocity = 0.0f;
+	float CameraDistanceAcceleration = 0.0f;
+
+#endif
+
+#if WITH_EDITOR
+
+public:
+
+	virtual FAnimGenControlObjectElement MakeControl_Implementation(FAnimGenControlObject& InControlObject, const FAnimDatabasePoseState& InPoseState) const override;
+
+	virtual void OnPreviewReset_Implementation(const FAnimDatabasePoseState& InPoseState) override;
+
+	virtual void OnPreviewUpdate_Implementation(const FAnimDatabasePoseState& InPoseState, const float DeltaTime) override;
+
+	virtual void OnPreviewInputKeyPressed_Implementation(const FKey Key) override;
+
+	virtual void OnPreviewInputKeyReleased_Implementation(const FKey Key) override;
+
+	virtual void OnPreviewDrawDebug_Implementation(
+		const FDebugDrawer& Drawer,
+		const FAnimGenControlObject& InControlObject,
+		const FAnimGenControlObjectElement InControlObjectElement,
+		const FAnimDatabasePoseState& InPoseState,
+		const FLinearColor& IdentifierColor) override;
+
+
+	virtual void OnPreviewInputLeftMousePressed_Implementation(const FVector MouseGroundLocation) override;
+
+	virtual void OnPreviewInputLeftMouseReleased_Implementation(const FVector MouseGroundLocation) override;
+
+	virtual void OnPreviewInputMouseMoved_Implementation(const FVector MouseGroundLocation) override;
+#endif
+};
+
+
+
+/** Behavior preview for Move To Target */
+UCLASS(BlueprintType, Blueprintable, meta = (DisplayName = "MoveToTarget"))
+class UAnimGenBehaviorPreview_MoveToTarget : public UAnimGenBehaviorPreview
+{
+	GENERATED_BODY()
+
+#if WITH_EDITORONLY_DATA
+
+public:
+
+	/** List of styles to cycle through */
+	UPROPERTY(EditAnywhere, Category = "Styles")
+	TArray<FName> StyleTags;
+
+protected:
+
+	/** Current Desired Style */
+	FName DesiredStyle = NAME_None;
+
+	/** Current Style Index */
+	int32 StyleIndex = 0;
+
+	/** Current Desired Target Location */
+	FVector TargetLocation = FVector::ZeroVector;
+
+	/** Current Desired Target Direction */
+	FVector TargetDirection = FVector::RightVector;
+
+	/** If to use the current target direction */
+	bool bUseTargetDirection = true;
+
+	/** If the active location is being placed */
+	bool bPlacingActive = false;
+
+	/** Current Active Location */
+	FVector ActiveLocation = FVector::ZeroVector;
+
+	/** Current Active Direction */
+	FVector ActiveDirection = FVector::RightVector;
+
+	/** If to use the current active direction */
+	bool bUseActiveDirection = true;
+
+#endif
+
+#if WITH_EDITOR
+
+public:
+
+	virtual FAnimGenControlObjectElement MakeControl_Implementation(FAnimGenControlObject& InControlObject, const FAnimDatabasePoseState& InPoseState) const override;
+
+	virtual void OnPreviewInputKeyPressed_Implementation(const FKey Key) override;
+
+	virtual void OnPreviewDrawDebug_Implementation(
+		const FDebugDrawer& Drawer,
+		const FAnimGenControlObject& InControlObject,
+		const FAnimGenControlObjectElement InControlObjectElement,
+		const FAnimDatabasePoseState& InPoseState,
+		const FLinearColor& IdentifierColor) override;
+
+	virtual void OnPreviewInputLeftMousePressed_Implementation(const FVector MouseGroundLocation) override;
+
+	virtual void OnPreviewInputLeftMouseReleased_Implementation(const FVector MouseGroundLocation) override;
+
+	virtual void OnPreviewInputMouseMoved_Implementation(const FVector MouseGroundLocation) override;
+#endif
+};
+
+#undef UE_API
\ No newline at end of file
diff --git a/Engine/Plugins/Experimental/Animation/AnimGen/Source/AnimGen/Public/AnimGenTasks.h b/Engine/Plugins/Experimental/Animation/AnimGen/Source/AnimGen/Public/AnimGenTasks.h
new file mode 100644
--- /dev/null
+++ b/Engine/Plugins/Experimental/Animation/AnimGen/Source/AnimGen/Public/AnimGenTasks.h
@@ -0,0 +1,49 @@
+// Copyright Epic Games, Inc. All Rights Reserved.
+
+#pragma once
+
+#include "CoreMinimal.h"
+#include "EvaluationVM/EvaluationTask.h"
+
+#include "AnimDatabasePose.h"
+
+#include "AnimGenTasks.generated.h"
+
+#define UE_API ANIMGEN_API
+
+/**
+ * FAnimGenPushPoseFromPoseState
+ *
+ * Pushes a pose to the stack from the given AnimDatabase pose state object
+ */
+USTRUCT()
+struct FAnimGenPushPoseFromPoseState : public FAnimNextEvaluationTask
+{
+	GENERATED_BODY()
+
+	DECLARE_ANIM_EVALUATION_TASK(FAnimGenPushPoseFromPoseState)
+
+	// Make a FAnimGenPushPoseFromPoseState task
+	static UE_API FAnimGenPushPoseFromPoseState Make(
+		const FAnimDatabasePoseState& PoseState,
+		const TArrayView<const FName> CurveAttributes,
+		const TArrayView<const FName> AttributeNames,
+		const TArrayView<const FName> BoneNames);
+
+	// Task entry point
+	UE_API virtual void Execute(UE::UAF::FEvaluationVM& VM) const override;
+
+	/** Names of the attribute to output as curves */
+	TArray<FName, TInlineAllocator<32>> CurveAttributes;
+
+	/** Names of all of the attributes in the AutoEncoder */
+	TArray<FName, TInlineAllocator<32>> AttributeNames;
+
+	/** PoseState to output */
+	FAnimDatabasePoseState PoseState;
+
+	/** Names of all the bones in the AutoEncoder */
+	TArray<FName> BoneNames;
+};
+
+#undef UE_API
\ No newline at end of file
diff --git a/Engine/Plugins/Experimental/Animation/AnimGen/Source/AnimGen/Public/AnimNode_AnimGenController.h b/Engine/Plugins/Experimental/Animation/AnimGen/Source/AnimGen/Public/AnimNode_AnimGenController.h
new file mode 100644
--- /dev/null
+++ b/Engine/Plugins/Experimental/Animation/AnimGen/Source/AnimGen/Public/AnimNode_AnimGenController.h
@@ -0,0 +1,149 @@
+// Copyright Epic Games, Inc. All Rights Reserved.
+
+#pragma once
+
+#include "Animation/AnimNodeBase.h"
+#include "Animation/AnimNotifies/AnimNotify.h"
+#include "Animation/AnimNotifies/AnimNotifyState.h"
+#include "Misc/FrameRate.h"
+
+#include "AnimDatabasePose.h"
+
+#include "AnimGenController.h"
+
+#include "AnimNode_AnimGenController.generated.h"
+
+#define UE_API ANIMGEN_API
+
+
+namespace UE::Learning
+{
+	struct FNeuralNetworkInference;
+}
+
+/**
+ * AnimNode that can be used to generate animation using an AnimGenController.
+ */
+USTRUCT(BlueprintInternalUseOnly, Experimental)
+struct FAnimNode_AnimGenController : public FAnimNode_Base
+{
+	GENERATED_BODY()
+
+	/** Controller asset used by the node */
+	UPROPERTY(EditAnywhere, Category = Settings, meta = (PinShownByDefault))
+	TObjectPtr<UAnimGenController> Controller;
+
+	/** Control Object taken as input */
+	UPROPERTY(EditAnywhere, Category = Settings, meta = (PinShownByDefault))
+	FAnimGenControlObject ControlObject;
+
+	/** Control Object Element from Control Object to use as input */
+	UPROPERTY(EditAnywhere, Category = Settings, meta = (PinShownByDefault))
+	FAnimGenControlObjectElement ControlObjectElement;
+
+	/** Minimum Smoothing time to use - mixes the predicted pose with an integration of the predicted velocity. */
+	UPROPERTY(EditAnywhere, Category = Settings, meta = (PinShownByDefault, ClampMin = "0.0", UIMin = "0.0"))
+	float MinPoseSmoothingTime = 0.05f;
+
+	/** Maximum Smoothing time to use - mixes the predicted pose with an integration of the predicted velocity. */
+	UPROPERTY(EditAnywhere, Category = Settings, meta = (PinShownByDefault, ClampMin = "0.0", UIMin = "0.0"))
+	float MaxPoseSmoothingTime = 0.25f;
+
+	/** Which LOD level to use */
+	UPROPERTY(EditAnywhere, Category = Settings, meta = (PinShownByDefault, ClampMin = "0", UIMin = "0", ClampMax = "2", UIMax = "2"))
+	int32 LODLevel = 0;
+
+public:
+
+	/** If to draw debug the pose state */
+	UPROPERTY(EditAnywhere, Category = Settings)
+	bool bDrawDebugPoseState = false;
+
+	/** If to get the root bone from the pose history */
+	UPROPERTY(EditAnywhere, Category = Settings)
+	bool bRootFromPoseHistory = true;
+
+	/** If to reset the node when it becomes relevant */
+	UPROPERTY(EditAnywhere, Category = Settings)
+	bool bResetOnBecomingRelevant = true;
+
+public:
+
+	/** Map of attribute names to anim notify states to produce */
+	UPROPERTY(EditAnywhere, Category = Notifies)
+	TMap<FName, TSubclassOf<UAnimNotifyState>> AnimNotifyStateMap;
+
+	/** Map of attribute names to anim notifies to produce */
+	UPROPERTY(EditAnywhere, Category = Notifies)
+	TMap<FName, TSubclassOf<UAnimNotify>> AnimNotifyMap;
+
+	/** Map of attribute names to curves to produce */
+	UPROPERTY(EditAnywhere, Category = Curves)
+	TMap<FName, FName> CurveMap;
+
+private:
+
+	/** Internal buffer of AnimNotifyEvents */
+	UPROPERTY(Transient)
+	TArray<FAnimNotifyEvent> AnimNotifyEvents;
+
+	/** Internal buffer of AnimNotifyEventReferences */
+	UPROPERTY(Transient)
+	TArray<FAnimNotifyEventReference> AnimNotifyEventReferences;
+
+public:
+
+	// FAnimNode_Base interface
+	virtual UE_API void Initialize_AnyThread(const FAnimationInitializeContext& Context) override;
+	virtual UE_API void Update_AnyThread(const FAnimationUpdateContext& Context) override;
+	virtual UE_API void CacheBones_AnyThread(const FAnimationCacheBonesContext& Context) override;
+	virtual UE_API void Evaluate_AnyThread(FPoseContext& Output) override;
+	// End of FAnimNode_Base interface
+
+private:
+
+	/** Update Counter for detecting being relevant */
+	FGraphTraversalCounter UpdateCounter;
+
+	/** Random state used for sampling initial noise */
+	uint32 RandomState = 1234;
+
+	/** Current active controller object */
+	UPROPERTY(Transient)
+	TObjectPtr<UAnimGenController> ActiveController = nullptr;
+
+	// Control Variables
+
+	bool bControlsValid = false;
+	int32 ControlVectorSize = 0;
+	int32 EncodedControlVectorSize = 0;
+	int32 ControlDistributionVectorSize = 0;
+	TLearningArray<2, float> ControlVector;
+	TLearningArray<2, float> EncodedControlVector;
+
+	// Pose Variables
+
+	bool bPoseStateRequiresReset = true;
+	FAnimDatabasePoseState PoseState;
+	int32 EncodedPoseVectorSize = 0;
+	int32 PoseVectorSize = 0;
+	TLearningArray<2, float> EncodedPoseVector;
+	TLearningArray<2, float> NextEncodedPoseVector;
+	TLearningArray<2, float> UnnormalizedEncodedPoseVector;
+	TLearningArray<2, float> PoseVector;
+	UE::AnimDatabase::FPoseData OutputPoseData;
+	TArray<TTuple<FName, float>> OutputCurves;
+
+	// Network Instances
+
+	TLearningArray<2, float> DenoiserInputVector;
+	TLearningArray<2, float> DenoiserOutputVector;
+	TSharedPtr<UE::Learning::FNeuralNetworkInference> ControlEncoderInference;
+	TSharedPtr<UE::Learning::FNeuralNetworkInference> LOD0Inference;
+	TSharedPtr<UE::Learning::FNeuralNetworkInference> LOD1Inference;
+	TSharedPtr<UE::Learning::FNeuralNetworkInference> LOD2Inference;
+	TSharedPtr<UE::Learning::FNeuralNetworkInference> EncoderInference;
+	TSharedPtr<UE::Learning::FNeuralNetworkInference> DecoderInference;
+};
+
+#undef UE_API
\ No newline at end of file
diff --git a/Engine/Plugins/Experimental/Animation/AnimGen/Source/AnimGen/Public/RigVMFunction_AnimGen.h b/Engine/Plugins/Experimental/Animation/AnimGen/Source/AnimGen/Public/RigVMFunction_AnimGen.h
new file mode 100644
--- /dev/null
+++ b/Engine/Plugins/Experimental/Animation/AnimGen/Source/AnimGen/Public/RigVMFunction_AnimGen.h
@@ -0,0 +1,200 @@
+// Copyright Epic Games, Inc. All Rights Reserved.
+
+#pragma once
+
+#include "RigVMCore/RigVMStruct.h"
+#include "RigVMFunctions/RigVMFunctionDefines.h"
+
+#include "Animation/TrajectoryTypes.h"
+#include "AnimGenControl.h"
+
+#include "RigVMFunction_AnimGen.generated.h"
+
+#define UE_API ANIMGEN_API
+
+/** Simple Object Representing a Name and Control Element Pair. */
+USTRUCT(BlueprintType)
+struct FAnimGenNamedControlElement
+{
+	GENERATED_BODY()
+
+	UPROPERTY(BlueprintReadWrite, EditAnywhere, Category = "Entry")
+	FName Name;
+
+	UPROPERTY(VisibleAnywhere, Category = "Entry")
+	FAnimGenControlObjectElement Element;
+};
+
+/*
+ * The base class for all AnimGen Nodes
+ */
+USTRUCT(meta=(Abstract, Category = "AnimGen|Experimental", NodeColor = "0.05 0.25 0.05", DocumentationPolicy="Strict", TemplateName = "AnimGen"))
+struct FRigVMFunction_AnimGenBase : public FRigVMStruct
+{
+	GENERATED_BODY()
+};
+
+/*
+ * Base class for all AnimGen rig units
+ */
+USTRUCT(meta = (Abstract, ExecuteContext = "FRigVMExecuteContext", Category = "AnimGen|Experimental", NodeColor = "0.05 0.25 0.05", DocumentationPolicy = "Strict"))
+struct FRigUnit_AnimGenBase : public FRigVMStruct
+{
+	GENERATED_BODY()
+
+	// The execution result
+	UPROPERTY(EditAnywhere, DisplayName = "Execute", Category = "BeginExecution", meta = (Input, Output))
+	FRigVMExecuteContext ExecuteContext;
+};
+
+/** Construct a new Control Object */
+USTRUCT(meta = (DisplayName = "Make Control Object"))
+struct FRigVMFunction_MakeControlObject : public FRigUnit_AnimGenBase
+{
+	GENERATED_BODY()
+
+	RIGVM_METHOD()
+	UE_API virtual void Execute() override;
+
+	// Output Control Object
+	UPROPERTY(meta = (Output))
+	FAnimGenControlObject Result;
+};
+
+/** Construct a Trajectory Control from a Transform Trajectory */
+USTRUCT(meta=(DisplayName="Make Trajectory Control"))
+struct FRigVMFunction_MakeTrajectoryControl : public FRigVMFunction_AnimGenBase
+{
+	GENERATED_BODY()
+
+	RIGVM_METHOD()
+	UE_API virtual void Execute() override;
+
+	// Control Object
+	UPROPERTY(meta = (Input, Output))
+	FAnimGenControlObject Object;
+
+	// Input Transform Trajectory
+	UPROPERTY(meta = (Input))
+	FTransformTrajectory TransformTrajectory;
+
+	// Relative Transform
+	UPROPERTY(meta = (Input))
+	FTransform RelativeTransform;
+
+	// Number of samples to take from the Trajectory
+	UPROPERTY(meta = (Input))
+	int32 SampleNum = 4;
+
+	// Time in the future to look for the trajectory
+	UPROPERTY(meta = (Input))
+	float FutureTime = 1.0f;
+
+	// Local forward vector of the trajectory
+	UPROPERTY(meta = (Input))
+	FVector ForwardVector = FVector(0,1,0);
+
+	// Tag for the control
+	UPROPERTY(meta = (Input))
+	FName Tag = TEXT("TrajectoryControl");
+
+	// Output Element
+	UPROPERTY(meta = (Output))
+	FAnimGenControlObjectElement Result;
+};
+
+/** Construct a Struct Control from the given named elements */
+USTRUCT(meta = (DisplayName = "Make Struct Control"))
+struct FRigVMFunction_MakeStructControl : public FRigVMFunction_AnimGenBase
+{
+	GENERATED_BODY()
+
+	RIGVM_METHOD()
+	UE_API virtual void Execute() override;
+
+	// Control Object
+	UPROPERTY(meta = (Input, Output))
+	FAnimGenControlObject Object;
+
+	// Struct Elements
+	UPROPERTY(meta = (Input))
+	TArray<FAnimGenNamedControlElement> Elements;
+
+	// Tag for the control
+	UPROPERTY(meta = (Input))
+	FName Tag = TEXT("StructControl");
+
+	// Output Element
+	UPROPERTY(meta = (Output))
+	FAnimGenControlObjectElement Result;
+};
+
+/** Construct a Named Inclusive Discrete Control from the given elements */
+USTRUCT(meta = (DisplayName = "Make Named Inclusive Discrete Control"))
+struct FRigVMFunction_MakeNamedInclusiveDiscreteControl : public FRigVMFunction_AnimGenBase
+{
+	GENERATED_BODY()
+
+	RIGVM_METHOD()
+	UE_API virtual void Execute() override;
+
+	// Control Object
+	UPROPERTY(meta = (Input, Output))
+	FAnimGenControlObject Object;
+
+	// Discrete Elements
+	UPROPERTY(meta = (Input))
+	TArray<FName> Elements;
+
+	// Tag for the control
+	UPROPERTY(meta = (Input))
+	FName Tag = TEXT("NamedInclusiveDiscreteControl");
+
+	// Output Element
+	UPROPERTY(meta = (Output))
+	FAnimGenControlObjectElement Result;
+};
+
+/** Construct a Control input for the Trajectory Following Behavior */
+USTRUCT(meta = (DisplayName = "Make Trajectory Follow Behavior Control"))
+struct FRigVMFunction_MakeTrajectoryFollowBehaviorControl : public FRigVMFunction_AnimGenBase
+{
+	GENERATED_BODY()
+
+	RIGVM_METHOD()
+	UE_API virtual void Execute() override;
+
+	// Control Object
+	UPROPERTY(meta = (Input, Output))
+	FAnimGenControlObject Object;
+
+	// Input Transform Trajectory
+	UPROPERTY(meta = (Input))
+	FTransformTrajectory TransformTrajectory;
+
+	// Relative Transform
+	UPROPERTY(meta = (Input))
+	FTransform RootBoneTransform;
+
+	// Desired Tags
+	UPROPERTY(meta = (Input))
+	TArray<FName> DesiredTags;
+
+	// Number of samples to take from the Trajectory
+	UPROPERTY(meta = (Input))
+	int32 TrajectorySampleNum = 4;
+
+	// Time in the future to look for the trajectory
+	UPROPERTY(meta = (Input))
+	float TrajectoryFutureTime = 1.0f;
+
+	// Local forward vector of the trajectory
+	UPROPERTY(meta = (Input))
+	FVector TrajectoryForwardVector = FVector(0, 1, 0);
+
+	// Output Element
+	UPROPERTY(meta = (Output))
+	FAnimGenControlObjectElement Result;
+};
+
+#undef UE_API
\ No newline at end of file
