AnimGen: Move to Experimental

diff --git a/Engine/Plugins/Experimental/Animation/AnimGen/Source/AnimGenEditor/Private/AnimGenEditorControllerMode.cpp b/Engine/Plugins/Experimental/Animation/AnimGen/Source/AnimGenEditor/Private/AnimGenEditorControllerMode.cpp
new file mode 100644
--- /dev/null
+++ b/Engine/Plugins/Experimental/Animation/AnimGen/Source/AnimGenEditor/Private/AnimGenEditorControllerMode.cpp
@@ -0,0 +1,2251 @@
+// Copyright Epic Games, Inc. All Rights Reserved.
+
+#include "AnimGenEditorControllerMode.h"
+#include "AnimGenEditorControllerToolkit.h"
+#include "SAnimGenEditorControlObject.h"
+
+#include "AnimDatabaseEditorViewportClient.h"
+#include "AnimDatabaseEditorPreviewScene.h"
+#include "AnimDatabaseEditorTimeline.h"
+
+#include "AnimGenLog.h"
+#include "AnimGenController.h"
+#include "AnimGenControl.h"
+#include "AnimGenAutoEncoder.h"
+#include "AnimGenBehavior.h"
+
+#include "AnimDatabase.h"
+#include "AnimDatabasePose.h"
+#include "AnimDatabaseFrameAttribute.h"
+#include "AnimDatabaseMath.h"
+
+#include "LearningNeuralNetwork.h"
+#include "LearningRandom.h"
+
+#include "GameFramework/Character.h"
+#include "Components/PoseableMeshComponent.h"
+#include "Components/CapsuleComponent.h"
+#include "Animation/Skeleton.h"
+#include "Engine/SkeletalMesh.h"
+#include "Async/ParallelFor.h"
+#include "EngineUtils.h"
+#include "SceneView.h"
+#include "Preferences/PersonaOptions.h"
+#include "DrawDebugLibrary.h"
+
+#define LOCTEXT_NAMESPACE "AnimGenEditorControllerMode"
+
+namespace UE::AnimGen::Editor
+{
+	namespace Controller::Private
+	{
+		static inline FVector GetRangeOffset(
+			const int32 CharacterIdx,
+			const int32 CharacterNum,
+			const float Spacing,
+			const bool bOffsetRanges,
+			const bool bOffsetRangesInGrid)
+		{
+			const int32 GridRows = FMath::CeilToInt(FMath::Sqrt((float)CharacterNum));
+
+			return
+				bOffsetRanges ?
+				bOffsetRangesInGrid ?
+				FVector(
+					Spacing * (CharacterIdx % GridRows) - 0.5f * Spacing * (GridRows - 1),
+					Spacing * (CharacterIdx / GridRows) - 0.5f * Spacing * (GridRows - 1), 0.0f) :
+				FVector(
+					Spacing * CharacterIdx - 0.5f * Spacing * (CharacterNum - 1), 0.0f, 0.0f) :
+				FVector::ZeroVector;
+		}
+
+		static inline void AdjustTransformRootTransform(
+			FVector& InOutLocation,
+			FQuat4f& InOutRotation,
+			FVector3f& InOutLinearVelocity,
+			FVector3f& InOutAngularVelocity,
+			const FVector Offset,
+			const FTransform RelativeTransform,
+			const bool bDisableRootTranslation,
+			const bool bDisableRootRotation)
+		{
+			const FTransform RootOffsetTransform = FTransform(FQuat::Identity, Offset, FVector::OneVector);
+			const FTransform RootAdjustmenTransform = RelativeTransform.Inverse() * RootOffsetTransform;
+
+			if (bDisableRootTranslation && bDisableRootRotation)
+			{
+				const FQuat4f RootRotationOffset = InOutRotation;
+				InOutLocation = RootOffsetTransform.GetLocation();
+				InOutRotation = FQuat4f::Identity;
+				InOutLinearVelocity = RootRotationOffset.UnrotateVector(InOutLinearVelocity);
+				InOutAngularVelocity = RootRotationOffset.UnrotateVector(InOutAngularVelocity);
+			}
+			else if (bDisableRootTranslation)
+			{
+				InOutLocation = RootOffsetTransform.GetLocation();
+				InOutRotation = ((FQuat4f)RootAdjustmenTransform.GetRotation()) * InOutRotation;
+				InOutLinearVelocity = (FVector3f)RootAdjustmenTransform.TransformVectorNoScale((FVector)InOutLinearVelocity);
+				InOutAngularVelocity = (FVector3f)RootAdjustmenTransform.TransformVectorNoScale((FVector)InOutAngularVelocity);
+			}
+			else if (bDisableRootRotation)
+			{
+				checkNoEntry();
+			}
+			else
+			{
+				InOutLocation = RootAdjustmenTransform.TransformPosition(InOutLocation);
+				InOutRotation = ((FQuat4f)RootAdjustmenTransform.GetRotation()) * InOutRotation;
+				InOutLinearVelocity = (FVector3f)RootAdjustmenTransform.TransformVectorNoScale((FVector)InOutLinearVelocity);
+				InOutAngularVelocity = (FVector3f)RootAdjustmenTransform.TransformVectorNoScale((FVector)InOutAngularVelocity);
+			}
+		}
+
+		static inline bool FindControlSetsTotalRange(int32& OutControlSetIdx, int32& OutRangeEntry, int32& OutRangeRange, const TArrayView<const FAnimGenControlSet> ControlSets, const int32 TotalRangeIdx)
+		{
+			const int32 ControlSetNum = ControlSets.Num();
+			int32 ControlSetRangeOffset = 0;
+			for (int32 ControlSetIdx = 0; ControlSetIdx < ControlSetNum; ControlSetIdx++)
+			{
+				const int32 ControlSetRangeNum = ControlSets[ControlSetIdx].Data->FrameRangeSet.GetTotalRangeNum();
+
+				if (TotalRangeIdx >= ControlSetRangeOffset && TotalRangeIdx < ControlSetRangeOffset + ControlSetRangeNum)
+				{
+					OutControlSetIdx = ControlSetIdx;
+					return ControlSets[ControlSetIdx].Data->FrameRangeSet.FindTotalRange(OutRangeEntry, OutRangeRange, TotalRangeIdx - ControlSetRangeOffset);
+				}
+
+				ControlSetRangeOffset += ControlSetRangeNum;
+			}
+
+			OutControlSetIdx = INDEX_NONE;
+			OutRangeEntry = INDEX_NONE;
+			OutRangeRange = INDEX_NONE;
+			return false;
+		}
+
+		static bool UpdateControlObjectTreeElement(
+			TSharedPtr<FControlObjectTreeElement> TreeElement,
+			const Learning::Observation::FObject& ControlObject, 
+			const Learning::Observation::FObjectElement& ControlObjectElement)
+		{
+			if (!ControlObject.IsValid(ControlObjectElement))
+			{
+				const bool bHasChildren = !TreeElement->Children.IsEmpty();
+				TreeElement->Type = EControlObjectTreeElementType::Invalid;
+				TreeElement->Children.Empty();
+				TreeElement->Name = NAME_None;
+				TreeElement->Tag = NAME_None;
+				TreeElement->ContinuousValues.Empty();
+				TreeElement->DiscreteValues.Empty();
+				return bHasChildren;
+			}
+
+			TreeElement->Tag = ControlObject.GetTag(ControlObjectElement);
+
+			switch (ControlObject.GetType(ControlObjectElement))
+			{
+			case Learning::Observation::EType::Null:
+			{
+				const bool bHasChildren = !TreeElement->Children.IsEmpty();
+				TreeElement->Type = EControlObjectTreeElementType::Null;
+				TreeElement->Children.Empty();
+				TreeElement->Name = NAME_None;
+				TreeElement->ContinuousValues.Empty();
+				TreeElement->DiscreteValues.Empty();
+				return bHasChildren;
+			}
+
+			case Learning::Observation::EType::Continuous:
+			{
+				const bool bHasChildren = !TreeElement->Children.IsEmpty();
+				TreeElement->Type = EControlObjectTreeElementType::Continuous;
+				TreeElement->Children.Empty();
+				TreeElement->Name = NAME_None;
+				TreeElement->ContinuousValues = ControlObject.GetContinuous(ControlObjectElement).Values;
+				TreeElement->DiscreteValues.Empty();
+				return bHasChildren;
+			}
+
+			case Learning::Observation::EType::NamedDiscreteExclusive:
+			{
+				const bool bHasChildren = !TreeElement->Children.IsEmpty();
+				TreeElement->Type = EControlObjectTreeElementType::NamedDiscreteExclusive;
+				TreeElement->Children.Empty();
+				TreeElement->Name = NAME_None;
+				TreeElement->ContinuousValues.Empty();
+				TreeElement->DiscreteValues.Empty(1);
+				TreeElement->DiscreteValues.Add(ControlObject.GetNamedDiscreteExclusive(ControlObjectElement).ElementName);
+				return bHasChildren;
+			}
+
+			case Learning::Observation::EType::NamedDiscreteInclusive:
+			{
+				const bool bHasChildren = !TreeElement->Children.IsEmpty();
+				TreeElement->Type = EControlObjectTreeElementType::NamedDiscreteInclusive;
+				TreeElement->Children.Empty();
+				TreeElement->Name = NAME_None;
+				TreeElement->ContinuousValues.Empty();
+				TreeElement->DiscreteValues = ControlObject.GetNamedDiscreteInclusive(ControlObjectElement).ElementNames;
+				return bHasChildren;
+			}
+
+			case Learning::Observation::EType::And:
+			{
+				TreeElement->Type = EControlObjectTreeElementType::And;
+				TreeElement->Name = NAME_None;
+				TreeElement->ContinuousValues.Empty();
+				TreeElement->DiscreteValues.Empty();
+
+				const TArrayView<const Learning::Observation::FObjectElement> SubElements = ControlObject.GetAnd(ControlObjectElement).Elements;
+				const TArrayView<const FName> SubElementNames = ControlObject.GetAnd(ControlObjectElement).ElementNames;
+				const int32 SubElementNum = SubElements.Num();
+
+				bool bRefresh = TreeElement->Children.Num() != SubElements.Num();
+				TreeElement->Children.SetNum(SubElements.Num());
+
+				for (int32 SubElementIdx = 0; SubElementIdx < SubElementNum; SubElementIdx++)
+				{
+					if (!TreeElement->Children[SubElementIdx].IsValid())
+					{
+						TreeElement->Children[SubElementIdx] = MakeShared<FControlObjectTreeElement>();
+						bRefresh = true;
+					}
+
+					bRefresh |= UpdateControlObjectTreeElement(TreeElement->Children[SubElementIdx], ControlObject, SubElements[SubElementIdx]);
+					TreeElement->Children[SubElementIdx]->Name = SubElementNames[SubElementIdx];
+				}
+
+				return bRefresh;
+			}
+
+			case Learning::Observation::EType::OrExclusive:
+			{
+				TreeElement->Type = EControlObjectTreeElementType::OrExclusive;
+				TreeElement->Name = NAME_None;
+				TreeElement->ContinuousValues.Empty();
+				TreeElement->DiscreteValues.Empty();
+
+				const Learning::Observation::FObjectElement SubElement = ControlObject.GetOrExclusive(ControlObjectElement).Element;
+				const FName SubElementName = ControlObject.GetOrExclusive(ControlObjectElement).ElementName;
+				
+				bool bRefresh = TreeElement->Children.Num() != 1;
+				TreeElement->Children.SetNum(1);
+				
+				if (!TreeElement->Children[0].IsValid())
+				{
+					TreeElement->Children[0] = MakeShared<FControlObjectTreeElement>();
+					bRefresh = true;
+				}
+
+				bRefresh |= UpdateControlObjectTreeElement(TreeElement->Children[0], ControlObject, SubElement);
+				TreeElement->Children[0]->Name = SubElementName;
+
+				return bRefresh;
+			}
+
+			case Learning::Observation::EType::OrInclusive:
+			{
+				TreeElement->Type = EControlObjectTreeElementType::OrInclusive;
+				TreeElement->Name = NAME_None;
+				TreeElement->ContinuousValues.Empty();
+				TreeElement->DiscreteValues.Empty();
+
+				const TArrayView<const Learning::Observation::FObjectElement> SubElements = ControlObject.GetOrInclusive(ControlObjectElement).Elements;
+				const TArrayView<const FName> SubElementNames = ControlObject.GetOrInclusive(ControlObjectElement).ElementNames;
+				const int32 SubElementNum = SubElements.Num();
+
+				bool bRefresh = TreeElement->Children.Num() != SubElements.Num();
+				TreeElement->Children.SetNum(SubElements.Num());
+
+				for (int32 SubElementIdx = 0; SubElementIdx < SubElementNum; SubElementIdx++)
+				{
+					if (!TreeElement->Children[SubElementIdx].IsValid())
+					{
+						TreeElement->Children[SubElementIdx] = MakeShared<FControlObjectTreeElement>();
+						bRefresh = true;
+					}
+
+					bRefresh |= UpdateControlObjectTreeElement(TreeElement->Children[SubElementIdx], ControlObject, SubElements[SubElementIdx]);
+					TreeElement->Children[SubElementIdx]->Name = SubElementNames[SubElementIdx];
+				}
+
+				return bRefresh;
+			}
+
+			case Learning::Observation::EType::Array:
+			{
+				TreeElement->Type = EControlObjectTreeElementType::Array;
+				TreeElement->Name = NAME_None;
+				TreeElement->ContinuousValues.Empty();
+				TreeElement->DiscreteValues.Empty();
+				
+				const TArrayView<const Learning::Observation::FObjectElement> SubElements = ControlObject.GetArray(ControlObjectElement).Elements;
+				const int32 SubElementNum = SubElements.Num();
+
+				bool bRefresh = TreeElement->Children.Num() != SubElements.Num();
+				TreeElement->Children.SetNum(SubElements.Num());
+
+				for (int32 SubElementIdx = 0; SubElementIdx < SubElementNum; SubElementIdx++)
+				{
+					if (!TreeElement->Children[SubElementIdx].IsValid())
+					{
+						TreeElement->Children[SubElementIdx] = MakeShared<FControlObjectTreeElement>();
+						bRefresh = true;
+					}
+
+					bRefresh |= UpdateControlObjectTreeElement(TreeElement->Children[SubElementIdx], ControlObject, SubElements[SubElementIdx]);
+				}
+
+				return bRefresh;
+			}
+
+			case Learning::Observation::EType::Set:
+			{
+				TreeElement->Type = EControlObjectTreeElementType::Set;
+				TreeElement->Name = NAME_None;
+				TreeElement->ContinuousValues.Empty();
+				TreeElement->DiscreteValues.Empty();
+
+				const TArrayView<const Learning::Observation::FObjectElement> SubElements = ControlObject.GetSet(ControlObjectElement).Elements;
+				const int32 SubElementNum = SubElements.Num();
+
+				bool bRefresh = TreeElement->Children.Num() != SubElements.Num();
+				TreeElement->Children.SetNum(SubElements.Num());
+
+				for (int32 SubElementIdx = 0; SubElementIdx < SubElementNum; SubElementIdx++)
+				{
+					if (!TreeElement->Children[SubElementIdx].IsValid())
+					{
+						TreeElement->Children[SubElementIdx] = MakeShared<FControlObjectTreeElement>();
+						bRefresh = true;
+					}
+
+					bRefresh |= UpdateControlObjectTreeElement(TreeElement->Children[SubElementIdx], ControlObject, SubElements[SubElementIdx]);
+				}
+
+				return bRefresh;
+			}
+
+			case Learning::Observation::EType::Encoding:
+			{
+				TreeElement->Type = EControlObjectTreeElementType::Encoding;
+				TreeElement->Name = NAME_None;
+				TreeElement->ContinuousValues.Empty();
+				TreeElement->DiscreteValues.Empty();
+
+				const Learning::Observation::FObjectElement SubElement = ControlObject.GetEncoding(ControlObjectElement).Element;
+
+				bool bRefresh = TreeElement->Children.Num() != 1;
+				TreeElement->Children.SetNum(1);
+
+				if (!TreeElement->Children[0].IsValid())
+				{
+					TreeElement->Children[0] = MakeShared<FControlObjectTreeElement>();
+					bRefresh = true;
+				}
+
+				bRefresh |= UpdateControlObjectTreeElement(TreeElement->Children[0], ControlObject, SubElement);
+
+				return bRefresh;
+			}
+
+			default:
+				const bool bHasChildren = !TreeElement->Children.IsEmpty();
+				TreeElement->Type = EControlObjectTreeElementType::Invalid;
+				TreeElement->Children.Empty();
+				TreeElement->Name = NAME_None;
+				TreeElement->ContinuousValues.Empty();
+				TreeElement->DiscreteValues.Empty();
+				return bHasChildren;
+			}
+
+		}
+	}
+
+	const FEditorModeID FControllerMode::EditorModeId = TEXT("AnimGenEditorControllerMode");
+
+	FControllerMode::FControllerMode() = default;
+	FControllerMode::~FControllerMode()
+	{
+		const int32 BehaviorNum = BehaviorPreviews.Num();
+		for (int32 BehaviorIdx = 0; BehaviorIdx < BehaviorNum; BehaviorIdx++)
+		{
+			if (BehaviorPreviews[BehaviorIdx] && BehaviorPreviews[BehaviorIdx]->IsValidLowLevel())
+			{
+				BehaviorPreviews[BehaviorIdx]->RemoveFromRoot();
+				BehaviorPreviews[BehaviorIdx] = nullptr;
+			}
+		}
+	}
+
+	bool FControllerMode::MouseMove(FEditorViewportClient* ViewportClient, FViewport* Viewport, int32 X, int32 Y)
+	{
+		if (!bCaptureInput)
+		{
+			return false;
+		}
+
+		FSceneViewFamilyContext ViewFamily(FSceneViewFamily::ConstructionValues(Viewport, ViewportClient->GetScene(), ViewportClient->EngineShowFlags).SetRealtimeUpdate(ViewportClient->IsRealtime()));
+		FSceneView* View = ViewportClient->CalcSceneView(&ViewFamily);
+		FViewportCursorLocation Cursor(View, ViewportClient, X, Y);
+
+		const FVector RayDirection = Cursor.GetDirection();
+
+		if (FMath::Abs(RayDirection.Z) < UE_SMALL_NUMBER)
+		{
+			return true;
+		}
+
+		const FVector RayStart = Cursor.GetOrigin();
+		const float Time = -RayStart.Z / RayDirection.Z;
+
+		if (Time > 0.0f)
+		{
+			const int32 BehaviorNum = BehaviorPreviews.Num();
+			for (int32 BehaviorIdx = 0; BehaviorIdx < BehaviorNum; BehaviorIdx++)
+			{
+				BehaviorPreviews[BehaviorIdx]->PreviewMouseGroundLocation = RayStart + Time * RayDirection;
+			}
+		}
+
+		return true;
+	}
+
+	bool FControllerMode::CapturedMouseMove(FEditorViewportClient* ViewportClient, FViewport* Viewport, int32 X, int32 Y)
+	{
+		return MouseMove(ViewportClient, Viewport, X, Y);
+	}
+
+	bool FControllerMode::InputAxis(FEditorViewportClient* InViewportClient, FViewport* Viewport, int32 ControllerId, FKey Key, float Delta, float DeltaTime)
+	{
+		if (!bCaptureInput)
+		{
+			return false;
+		}
+
+		const int32 BehaviorNum = BehaviorPreviews.Num();
+
+		if (Key == EKeys::Gamepad_LeftX)
+		{
+			for (int32 BehaviorIdx = 0; BehaviorIdx < BehaviorNum; BehaviorIdx++)
+			{
+				BehaviorPreviews[BehaviorIdx]->PreviewLeftStick.Y = Delta;
+			}
+			return true;
+		}
+		else if (Key == EKeys::Gamepad_LeftY)
+		{
+			for (int32 BehaviorIdx = 0; BehaviorIdx < BehaviorNum; BehaviorIdx++)
+			{
+				BehaviorPreviews[BehaviorIdx]->PreviewLeftStick.X = Delta;
+			}
+			return true;
+		}
+		else if (Key == EKeys::Gamepad_RightX)
+		{
+			for (int32 BehaviorIdx = 0; BehaviorIdx < BehaviorNum; BehaviorIdx++)
+			{
+				BehaviorPreviews[BehaviorIdx]->PreviewRightStick.Y = Delta;
+			}
+			return true;
+		}
+		else if (Key == EKeys::Gamepad_RightY)
+		{
+			for (int32 BehaviorIdx = 0; BehaviorIdx < BehaviorNum; BehaviorIdx++)
+			{
+				BehaviorPreviews[BehaviorIdx]->PreviewRightStick.X = -Delta;
+			}
+			return true;
+		}
+		else if (Key == EKeys::Gamepad_LeftTriggerAxis)
+		{
+			for (int32 BehaviorIdx = 0; BehaviorIdx < BehaviorNum; BehaviorIdx++)
+			{
+				BehaviorPreviews[BehaviorIdx]->PreviewLeftTrigger = Delta;
+			}
+			return true;
+		}
+		else if (Key == EKeys::Gamepad_RightTriggerAxis)
+		{
+			for (int32 BehaviorIdx = 0; BehaviorIdx < BehaviorNum; BehaviorIdx++)
+			{
+				BehaviorPreviews[BehaviorIdx]->PreviewRightTrigger = Delta;
+			}
+			return true;
+		}
+
+		return false;
+	}
+
+	bool FControllerMode::InputKey(FEditorViewportClient* ViewportClient, FViewport* Viewport, FKey Key, EInputEvent Event)
+	{
+		if (!bCaptureInput)
+		{
+			return false;
+		}
+
+		KeyEvents.Add({ Key, Event });
+		return true;
+	}
+
+	void FControllerMode::TickTraining(FEditorViewportClient* EditorViewportClient, float DeltaTime)
+	{
+		bCaptureInput = false;
+
+		AnimDatabase::Editor::FViewportClient* ViewportClient = static_cast<AnimDatabase::Editor::FViewportClient*>(EditorViewportClient);
+		check(ViewportClient);
+
+		TSharedPtr<FControllerToolkit> EditorToolKit = StaticCastSharedPtr<FControllerToolkit>(ViewportClient->GetAssetEditorToolkit().Pin());
+		if (!EditorToolKit.IsValid()) { return; }
+
+		WeakToolkit = EditorToolKit.ToWeakPtr();
+
+		TSharedPtr<UE::AnimDatabase::Editor::FPreviewScene> PreviewScene = ViewportClient->GetPreviewScene().Pin();
+		if (!PreviewScene.IsValid()) { return; }
+
+		// Get the Controller
+
+		UAnimGenController* Controller = EditorToolKit->Controller;
+		check(Controller);
+
+		// Update Training Settings
+
+		Controller->TrainingSettings->RangeIndentifierColorSeed = Controller->ViewportSettings->RangeIdentifierColorSeed;
+		const bool bTrainingSettingsUpdated = Controller->TrainingSettings->Update();
+
+		if (bTrainingSettingsUpdated)
+		{
+			ActiveRanges.Reset();
+			EditorToolKit->TrainingSettingsWidget->ForceRefresh();
+		}
+
+		// Load Database
+
+		UAnimDatabase* TrainingDatabase = Controller->TrainingSettings->Database.Get();
+		if (!TrainingDatabase) { return; }
+
+		// Load AutoEncoder
+
+		UAnimGenAutoEncoder* TrainingAutoEncoder = Controller->TrainingSettings->AutoEncoder.Get();
+		if (!TrainingAutoEncoder || !TrainingAutoEncoder->IsValid()) { return; }
+
+		// Load Behavior
+
+		//UClass* TrainingBehavior = Controller->TrainingSettings->Behavior;
+		//if (!TrainingBehavior) { return; }
+
+		// Add the required number of characters to the scene
+
+		while (PreviewScene->GetCharacterNum() < MaxCharacterNum)
+		{
+			const int32 CharacterIdx = PreviewScene->AddCharacter();
+			PreviewScene->SetCharacterVisibility(CharacterIdx, false);
+		}
+
+		// Set the character's skeletal mesh based on what is given in the animation database
+
+		USkinnedAsset* DesiredSkeletalMesh = TrainingDatabase->Skeleton ? TrainingDatabase->Skeleton->GetPreviewMesh() : nullptr;
+		DesiredSkeletalMesh = Controller->ViewportSettings->PreviewMesh ? Controller->ViewportSettings->PreviewMesh.Get() : DesiredSkeletalMesh;
+
+		for (int32 CharacterIdx = 0; CharacterIdx < MaxCharacterNum; CharacterIdx++)
+		{
+			USkinnedAsset* CurrentSkeletalMesh = PreviewScene->GetPoseableMeshComponent(CharacterIdx)->GetSkinnedAsset();
+
+			if (DesiredSkeletalMesh && DesiredSkeletalMesh != CurrentSkeletalMesh)
+			{
+				PreviewScene->GetPoseableMeshComponent(CharacterIdx)->SetSkinnedAssetAndUpdate(DesiredSkeletalMesh, true);
+			}
+		}
+
+		// Compute all range trajectories
+
+		int32 QueryRangeNum = 0;
+
+		const int32 ControlSetNum = Controller->TrainingSettings->ControlSets.Num();
+		for (int32 ControlSetIdx = 0; ControlSetIdx < ControlSetNum; ControlSetIdx++)
+		{
+			QueryRangeNum += Controller->TrainingSettings->ControlSets[ControlSetIdx].Data->FrameRangeSet.GetTotalRangeNum();
+		}
+
+		if (bTrainingSettingsUpdated || RangeTrajectoryLocations.IsEmpty())
+		{
+			RangeTrajectoryLocations.SetNum(QueryRangeNum);
+			RangeTrajectoryRotations.SetNum(QueryRangeNum);
+
+			ParallelFor(QueryRangeNum, [this, TrainingDatabase, Controller](int32 RangeIdx)
+				{
+					int32 ControlSetIdx = INDEX_NONE, RangeEntry = INDEX_NONE, RangeRange = INDEX_NONE;
+					const bool bFound = Controller::Private::FindControlSetsTotalRange(ControlSetIdx, RangeEntry, RangeRange, Controller->TrainingSettings->ControlSets, RangeIdx);
+					check(bFound);
+
+					const int32 RangeFrameNum = Controller->TrainingSettings->ControlSets[ControlSetIdx].Data->FrameRangeSet.GetEntryRangeLength(RangeEntry, RangeRange);
+					const int32 RangeSequence = Controller->TrainingSettings->ControlSets[ControlSetIdx].Data->FrameRangeSet.GetEntrySequence(RangeEntry);
+					const int32 RangeStart = Controller->TrainingSettings->ControlSets[ControlSetIdx].Data->FrameRangeSet.GetEntryRangeStart(RangeEntry, RangeRange);
+
+					RangeTrajectoryLocations[RangeIdx].SetNumUninitialized({ RangeFrameNum });
+					RangeTrajectoryRotations[RangeIdx].SetNumUninitialized({ RangeFrameNum });
+
+					TrainingDatabase->GetRootLocation(RangeTrajectoryLocations[RangeIdx], RangeSequence, RangeStart);
+					TrainingDatabase->GetRootRotation(RangeTrajectoryRotations[RangeIdx], RangeSequence, RangeStart);
+				});
+		}
+
+		// Adjust Number of visible characters in the scene
+
+		CharacterNum = FMath::Min(Controller->TrainingSettings->SelectedRanges.Num(), MaxCharacterNum);
+
+		for (int32 CharacterIdx = 0; CharacterIdx < MaxCharacterNum; CharacterIdx++)
+		{
+			PreviewScene->SetCharacterVisibility(CharacterIdx, CharacterIdx < CharacterNum);
+		}
+
+		// If selection has changed adjust view ranges
+
+		if (bTrainingSettingsUpdated || ActiveRanges != Controller->TrainingSettings->SelectedRanges)
+		{
+			ActiveRanges = Controller->TrainingSettings->SelectedRanges;
+
+			// Add error message if too many ranges selected
+
+			if (ActiveRanges.Num() > MaxCharacterNum)
+			{
+				ViewportClient->SetWarningMessage(FText::Format(
+					LOCTEXT("AnimGenTooManyRangesWarning", "Warning: Rendering characters for {0} for {1} selected ranges."),
+					FText::AsNumber(MaxCharacterNum),
+					FText::AsNumber(ActiveRanges.Num())));
+			}
+			else
+			{
+				ViewportClient->SetWarningMessage(FText());
+			}
+
+			// Adjust the timeline for the newly selected ranges
+
+			if (ActiveRanges.Num() == 0)
+			{
+				EditorToolKit->TimelineModel->SetTimelineFrameRate(TrainingDatabase->GetFrameRate());
+				EditorToolKit->TimelineModel->SetTimelinePlaybackRange(TRange<FFrameNumber>(0, 0));
+				EditorToolKit->TimelineModel->SetTimelineWorkingRange(TRange<FFrameNumber>(0, 0));
+				EditorToolKit->TimelineModel->SetTimelineViewRange(TRange<FFrameNumber>(0, 0));
+				EditorToolKit->TimelineModel->SetTimelineFrameTime(FFrameTime::FromDecimal(0));
+			}
+			else if (ActiveRanges.Num() == 1)
+			{
+				int32 ControlSetIdx = INDEX_NONE, RangeEntry = INDEX_NONE, RangeRange = INDEX_NONE;
+				const bool bFound = Controller::Private::FindControlSetsTotalRange(ControlSetIdx, RangeEntry, RangeRange, Controller->TrainingSettings->ControlSets, ActiveRanges[0]);
+				check(bFound);
+
+				const int32 RangeStart = Controller->TrainingSettings->ControlSets[ControlSetIdx].Data->FrameRangeSet.GetEntryRangeStart(RangeEntry, RangeRange);
+				const int32 RangeLength = Controller->TrainingSettings->ControlSets[ControlSetIdx].Data->FrameRangeSet.GetEntryRangeLength(RangeEntry, RangeRange);
+
+				EditorToolKit->TimelineModel->SetTimelineFrameRate(TrainingDatabase->GetFrameRate());
+				EditorToolKit->TimelineModel->SetTimelinePlaybackRange(TRange<FFrameNumber>(RangeStart, RangeStart + RangeLength - 1));
+				EditorToolKit->TimelineModel->SetTimelineWorkingRange(TRange<FFrameNumber>(RangeStart, RangeStart + RangeLength - 1));
+				EditorToolKit->TimelineModel->SetTimelineViewRange(TRange<FFrameNumber>(RangeStart, RangeStart + RangeLength - 1));
+				EditorToolKit->TimelineModel->SetTimelineFrameTime(FFrameTime::FromDecimal(RangeStart));
+			}
+			else
+			{
+				int32 RangeLength = 1;
+				for (const int32 ActiveRangeIdx : ActiveRanges)
+				{
+					int32 ControlSetIdx = INDEX_NONE, RangeEntry = INDEX_NONE, RangeRange = INDEX_NONE;
+					const bool bFound = Controller::Private::FindControlSetsTotalRange(ControlSetIdx, RangeEntry, RangeRange, Controller->TrainingSettings->ControlSets, ActiveRangeIdx);
+					check(bFound);
+					RangeLength = FMath::Max(Controller->TrainingSettings->ControlSets[ControlSetIdx].Data->FrameRangeSet.GetEntryRangeLength(RangeEntry, RangeRange), RangeLength);
+				}
+
+				EditorToolKit->TimelineModel->SetTimelineFrameRate(TrainingDatabase->GetFrameRate());
+				EditorToolKit->TimelineModel->SetTimelinePlaybackRange(TRange<FFrameNumber>(0, RangeLength - 1));
+				EditorToolKit->TimelineModel->SetTimelineWorkingRange(TRange<FFrameNumber>(0, RangeLength - 1));
+				EditorToolKit->TimelineModel->SetTimelineViewRange(TRange<FFrameNumber>(0, RangeLength - 1));
+				EditorToolKit->TimelineModel->SetTimelineFrameTime(FFrameTime::FromDecimal(0));
+			}
+		}
+
+		// Update Character Range Data
+
+		CharacterSequenceIndices.SetNumUninitialized(CharacterNum);
+		CharacterRangeTimes.SetNumUninitialized(CharacterNum);
+		CharacterRangeStarts.SetNumUninitialized(CharacterNum);
+		CharacterRangeLengths.SetNumUninitialized(CharacterNum);
+		CharacterControlSetIndices.SetNumUninitialized(CharacterNum);
+		CharacterControlSetEntries.SetNumUninitialized(CharacterNum);
+		CharacterControlSetRanges.SetNumUninitialized(CharacterNum);
+		CharacterControlSetFrames.SetNumUninitialized(CharacterNum);
+
+		for (int32 CharacterIdx = 0; CharacterIdx < CharacterNum; CharacterIdx++)
+		{
+			int32 ControlSetIdx = INDEX_NONE, RangeEntry = INDEX_NONE, RangeRange = INDEX_NONE;
+			const bool bFound = Controller::Private::FindControlSetsTotalRange(ControlSetIdx, RangeEntry, RangeRange, Controller->TrainingSettings->ControlSets, ActiveRanges[CharacterIdx]);
+			check(bFound);
+
+			const UE::Learning::FFrameRangeSet& FrameRangeSet = Controller->TrainingSettings->ControlSets[ControlSetIdx].Data->FrameRangeSet;
+
+			CharacterSequenceIndices[CharacterIdx] = FrameRangeSet.GetEntrySequence(RangeEntry);
+			CharacterRangeStarts[CharacterIdx] = FrameRangeSet.GetEntryRangeStart(RangeEntry, RangeRange);
+			CharacterRangeLengths[CharacterIdx] = FrameRangeSet.GetEntryRangeLength(RangeEntry, RangeRange);
+
+			CharacterRangeTimes[CharacterIdx] = 0.0;
+			if (ActiveRanges.Num() == 1)
+			{
+				CharacterRangeTimes[CharacterIdx] = EditorToolKit->TimelineModel->GetTimelineFrameTime().AsDecimal() / TrainingDatabase->GetFrameRate().AsDecimal();
+			}
+			else
+			{
+				CharacterRangeTimes[CharacterIdx] = FMath::Clamp(
+					EditorToolKit->TimelineModel->GetTimelineFrameTime().AsDecimal() + (float)CharacterRangeStarts[CharacterIdx],
+					(float)CharacterRangeStarts[CharacterIdx],
+					(float)(CharacterRangeStarts[CharacterIdx] + CharacterRangeLengths[CharacterIdx] - 1)) / TrainingDatabase->GetFrameRate().AsDecimal();
+			}
+
+			// Clamp some small epsilon within the duration to help with sampling of frame attributes.
+			CharacterRangeTimes[CharacterIdx] = FMath::Clamp(CharacterRangeTimes[CharacterIdx],
+				0.0f, FMath::Max(TrainingDatabase->GetSequenceDuration(CharacterSequenceIndices[CharacterIdx]) - UE_KINDA_SMALL_NUMBER, 0.0f));
+
+			// Find current control set and nearest frame in that control set
+			CharacterControlSetIndices[CharacterIdx] = ControlSetIdx;
+			CharacterControlSetEntries[CharacterIdx] = RangeEntry;
+			CharacterControlSetRanges[CharacterIdx] = RangeRange;
+			CharacterControlSetFrames[CharacterIdx] = FMath::Clamp(
+				FMath::RoundToInt(CharacterRangeTimes[CharacterIdx] * TrainingDatabase->GetFrameRate().AsDecimal()) - CharacterRangeStarts[CharacterIdx],
+				0, CharacterRangeLengths[CharacterIdx] - 1);
+
+			check(TrainingDatabase->GetAnimSequence(CharacterSequenceIndices[CharacterIdx]))
+		}
+
+		// Resize Pose State and re-create Network Inference Instances
+
+		const int32 BoneNum = TrainingDatabase->GetBoneNum();
+
+		if (!PoseState.IsValid())
+		{
+			PoseState = UAnimDatabasePoseStateLibrary::MakePoseState();
+		}
+
+		TArray<int32> BoneParents;
+		TrainingDatabase->GetBoneParents(BoneParents);
+
+		PoseState.CurrPoseData->Resize(CharacterNum, BoneNum, TrainingAutoEncoder->AttributeTypes);
+		PoseState.PoseGlobalBoneData->Resize(CharacterNum, BoneNum);
+
+		// Frame Attributes
+
+		const int32 DatabaseContentHash = TrainingDatabase->GetContentHash();
+
+		if (FrameAttributeEntries != TrainingAutoEncoder->TrainedFrameAttributes ||
+			FrameAttributesDatabaseContentHash != DatabaseContentHash)
+		{
+			FScopedSlowTask SlowTask(0.0f, LOCTEXT("RebuildingFrameAttributes", "Rebuilding AutoEncoder Frame Attributes in Database"));
+			SlowTask.MakeDialog();
+
+			FrameAttributeEntries = TrainingAutoEncoder->TrainedFrameAttributes;
+			FrameAttributesDatabaseContentHash = DatabaseContentHash;
+
+			const int32 AttributeNum = FrameAttributeEntries.Num();
+			FrameAttributes.SetNum(AttributeNum);
+			for (int32 AttributeIdx = 0; AttributeIdx < AttributeNum; AttributeIdx++)
+			{
+				FrameAttributes[AttributeIdx] = UAnimDatabaseFrameAttributeLibrary::MakeFrameAttributeFromFunction(
+					TrainingDatabase, 
+					Controller->TrainingSettings->QueryRanges,
+					FrameAttributeEntries[AttributeIdx].FrameAttribute);
+
+				if (!FrameAttributes[AttributeIdx].IsValid())
+				{
+					FrameAttributes[AttributeIdx] = UAnimDatabaseFrameAttributeLibrary::MakeNullFrameAttribute(Controller->TrainingSettings->QueryRanges);
+				}
+			}
+		}
+
+		// Evaluate Pose and attributes
+
+		TrainingDatabase->WaitForCompressionOnAnimSequences(CharacterSequenceIndices);
+
+		TrainingDatabase->SamplePoseData(
+			PoseState.CurrPoseData->View(),
+			CharacterSequenceIndices,
+			CharacterRangeTimes,
+			FrameAttributes);
+
+		if (Controller->ViewportSettings->bRangesStartAtOrigin)
+		{
+			for (int32 CharacterIdx = 0; CharacterIdx < CharacterNum; CharacterIdx++)
+			{
+				const FVector RootOffset = Controller::Private::GetRangeOffset(
+					CharacterIdx, CharacterNum,
+					Controller->ViewportSettings->RangeOffsetSpacing,
+					Controller->ViewportSettings->bOffsetRanges,
+					Controller->ViewportSettings->bOffsetRangesInGrid);
+
+				FTransform RangeStartRootTransform;
+				TrainingDatabase->GetRootTransform(TLearningArrayView<1, FTransform>(&RangeStartRootTransform, 1),
+					CharacterSequenceIndices[CharacterIdx], CharacterRangeStarts[CharacterIdx]);
+
+				Controller::Private::AdjustTransformRootTransform(
+					PoseState.CurrPoseData->RootData.RootLocations[CharacterIdx],
+					PoseState.CurrPoseData->RootData.RootRotations[CharacterIdx],
+					PoseState.CurrPoseData->RootData.RootLinearVelocities[CharacterIdx],
+					PoseState.CurrPoseData->RootData.RootAngularVelocities[CharacterIdx],
+					RootOffset,
+					RangeStartRootTransform,
+					Controller->ViewportSettings->bDisableRootTranslation,
+					Controller->ViewportSettings->bDisableRootRotation);
+			}
+		}
+
+		// Compute Forward Kinematics
+
+		AnimDatabase::PoseData::ForwardKinematics(
+			PoseState.PoseGlobalBoneData->View(),
+			PoseState.CurrPoseData->LocalBoneData.ConstView(),
+			PoseState.CurrPoseData->RootData.ConstView(),
+			BoneParents);
+
+		// Update Pose
+
+		for (int32 CharacterIdx = 0; CharacterIdx < CharacterNum; CharacterIdx++)
+		{
+			if (PreviewScene->GetPoseableMeshComponent(CharacterIdx)->GetSkinnedAsset())
+			{
+				for (int32 BoneIdx = 0; BoneIdx < BoneNum; BoneIdx++)
+				{
+					const FTransform BoneTransform = FTransform(
+						((FQuat)PoseState.PoseGlobalBoneData->BoneRotations[CharacterIdx][BoneIdx]).GetNormalized(),
+						(FVector)PoseState.PoseGlobalBoneData->BoneLocations[CharacterIdx][BoneIdx],
+						(FVector)PoseState.PoseGlobalBoneData->BoneScales[CharacterIdx][BoneIdx]);
+
+					PreviewScene->GetPoseableMeshComponent(CharacterIdx)->SetBoneTransformByName(TrainingDatabase->GetBoneName(BoneIdx), BoneTransform, EBoneSpaces::WorldSpace);
+				}
+
+				PreviewScene->GetPoseableMeshComponent(CharacterIdx)->RefreshBoneTransforms();
+			}
+		}
+
+		// Update UI Control Objects
+
+		if (Controller->TrainingSettings->ControlObject.ObservationObject.IsValid() && CharacterNum > 0)
+		{
+			TArray<TSharedPtr<FControlObjectTreeElement>>& ControlObjectTreeElements = EditorToolKit->ControlObjectWidget->GetTreeElementsRef();
+
+			bool bRefreshControlUI = ControlObjectTreeElements.Num() != CharacterNum;
+
+			ControlObjectTreeElements.SetNum(CharacterNum);
+
+			for (int32 CharacterIdx = 0; CharacterIdx < CharacterNum; CharacterIdx++)
+			{
+				if (!ControlObjectTreeElements[CharacterIdx].IsValid())
+				{
+					ControlObjectTreeElements[CharacterIdx] = MakeShared<FControlObjectTreeElement>();
+					bRefreshControlUI = true;
+				}
+
+				const int32 FlatIndex = Controller->TrainingSettings->ControlSets[CharacterControlSetIndices[CharacterIdx]].Data->FrameRangeSet.GetEntryRangeOffset(
+					CharacterControlSetEntries[CharacterIdx], CharacterControlSetRanges[CharacterIdx]) + CharacterControlSetFrames[CharacterIdx];
+
+				bRefreshControlUI |= Controller::Private::UpdateControlObjectTreeElement(
+					ControlObjectTreeElements[CharacterIdx],
+					Controller->TrainingSettings->ControlObject.ObservationObject->Object,
+					Controller->TrainingSettings->ControlSets[CharacterControlSetIndices[CharacterIdx]].Data->Controls[FlatIndex].ObjectElement);
+
+				ControlObjectTreeElements[CharacterIdx]->Color = Controller->TrainingSettings->QueryEntries[ActiveRanges[CharacterIdx]]->Color;
+			}
+
+			if (bRefreshControlUI)
+			{
+				EditorToolKit->ControlObjectWidget->RefreshTreeElements();
+			}
+		}
+		else
+		{
+			TArray<TSharedPtr<FControlObjectTreeElement>>& ControlObjectTreeElements = EditorToolKit->ControlObjectWidget->GetTreeElementsRef();
+			if (!ControlObjectTreeElements.IsEmpty())
+			{
+				ControlObjectTreeElements.Empty();
+				EditorToolKit->ControlObjectWidget->RefreshTreeElements();
+			}
+		}
+
+		// Update test range time if playing animation
+
+		EditorToolKit->TimelineModel->UpdateTimelineFrameTime(DeltaTime * TrainingDatabase->GetFrameRate().AsDecimal());
+	}
+
+	void FControllerMode::TickTesting(FEditorViewportClient* EditorViewportClient, float DeltaTime)
+	{
+		AnimDatabase::Editor::FViewportClient* ViewportClient = static_cast<AnimDatabase::Editor::FViewportClient*>(EditorViewportClient);
+		check(ViewportClient);
+		ViewportClient->ClearWarningMessage();
+
+		TSharedPtr<FControllerToolkit> EditorToolKit = StaticCastSharedPtr<FControllerToolkit>(ViewportClient->GetAssetEditorToolkit().Pin());
+		if (!EditorToolKit.IsValid()) { return; }
+
+		WeakToolkit = EditorToolKit.ToWeakPtr();
+
+		TSharedPtr<UE::AnimDatabase::Editor::FPreviewScene> PreviewScene = ViewportClient->GetPreviewScene().Pin();
+		if (!PreviewScene.IsValid()) { return; }
+
+		// Get the Controller
+
+		UAnimGenController* Controller = EditorToolKit->Controller;
+		check(Controller);
+
+		// Check we have a trained controller
+
+		if (!Controller->IsValid()) { return; }
+
+		// Check we have a preview class
+
+		TObjectPtr<UAnimGenBehaviorPreview> BehaviorPreview = Controller->TestingSettings->BehaviorPreview;
+
+		if (!BehaviorPreview)
+		{
+			return;
+		}
+
+		// We need to use a fixed timestep because the controller is only trained on animation data at 60fps
+
+		if (Controller->TestingSettings->bUseFixedTimeStep)
+		{
+			DeltaTime = 1.0f / FMath::Max(Controller->TestingSettings->FixedTimeStep.AsDecimal(), UE_SMALL_NUMBER);
+		}
+
+		// Add the required number of characters to the scene
+
+		while (PreviewScene->GetCharacterNum() < MaxCharacterNum)
+		{
+			const int32 CharacterIdx = PreviewScene->AddCharacter();
+			PreviewScene->SetCharacterVisibility(CharacterIdx, false);
+		}
+
+		// Init Character Colors
+
+		if (CharacterColors.Num() != MaxCharacterNum)
+		{
+			CharacterColors.SetNumUninitialized(MaxCharacterNum);
+
+			const uint32 SeedHash = Learning::Random::Int(Controller->ViewportSettings->RangeIdentifierColorSeed);
+			for (int32 CharacterIdx = 0; CharacterIdx < MaxCharacterNum; CharacterIdx++)
+			{
+				const uint32 ColorHash = Learning::Random::Int(0x94f67075 ^ CharacterIdx ^ SeedHash);
+				const uint8 Hue = ColorHash & 0x000000FF;
+				const uint8 Saturation = ((ColorHash & 0x0000FF00) >> 8) % 25 + 230;
+				const uint8 Brightness = ((ColorHash & 0x00FF0000) >> 16) % 50 + 205;
+				CharacterColors[CharacterIdx] = FLinearColor::MakeFromHSV8(Hue, Saturation, Brightness);
+			}
+		}
+
+		// Load the AutoEncoder 
+
+		UAnimGenAutoEncoder* AutoEncoder = Controller->AutoEncoder;
+		if (!AutoEncoder || !AutoEncoder->IsValid()) { return; }
+
+		//TSubclassOf<UAnimGenBehavior> BehaviorClass = Controller->TrainedBehavior;
+		//if (!BehaviorClass) { return; }
+
+		//const UAnimGenBehavior* Behavior = BehaviorClass->GetDefaultObject<UAnimGenBehavior>();
+
+		const int32 CurrentAutoEncoderContentHash = AutoEncoder->GetContentHash();
+		if (CurrentAutoEncoderContentHash != TrainedAutoEncoderContentHash)
+		{
+			TrainedAutoEncoderContentHash = CurrentAutoEncoderContentHash;
+			EditorToolKit->EditingAssetWidget->ForceRefresh();
+		}
+
+		const int32 BoneNum = AutoEncoder->GetBoneNum();
+		const int32 PoseVectorSize = AutoEncoder->GetPoseVectorSize();
+		const int32 EncodedPoseVectorSize = AutoEncoder->GetEncodingSize();
+
+		// TODO: Check Encoding Size Matches
+
+		// Set the character's skeletal mesh based on what is given in the animation database
+
+		USkinnedAsset* DesiredSkeletalMesh = AutoEncoder->TrainedSkeleton ? AutoEncoder->TrainedSkeleton->GetPreviewMesh() : nullptr;
+		DesiredSkeletalMesh = Controller->ViewportSettings->PreviewMesh ? Controller->ViewportSettings->PreviewMesh.Get() : DesiredSkeletalMesh;
+
+		for (int32 CharacterIdx = 0; CharacterIdx < MaxCharacterNum; CharacterIdx++)
+		{
+			USkinnedAsset* CurrentSkeletalMesh = PreviewScene->GetPoseableMeshComponent(CharacterIdx)->GetSkinnedAsset();
+
+			if (DesiredSkeletalMesh && DesiredSkeletalMesh != CurrentSkeletalMesh)
+			{
+				PreviewScene->GetPoseableMeshComponent(CharacterIdx)->SetSkinnedAssetAndUpdate(DesiredSkeletalMesh, true);
+			}
+		}
+
+		// Check if the trained database has changed and if so force refresh of details panel to check if warning needs to be displayed
+
+		/*
+		UAnimDatabase* TrainedDatabase = Controller->TrainedDatabase.LoadSynchronous();
+		if (!TrainedDatabase) { return; }
+
+		const int32 CurrentTrainedDatabaseContentHash = TrainedDatabase->GetContentHash();
+		UObject* CurrentTrainedFrameRangesDefaultObject =
+			!Controller->TrainedFrameRanges ? nullptr :
+			Controller->TrainedFrameRanges->GetDefaultObject();
+
+		if (CurrentTrainedDatabaseContentHash != TrainedDatabaseContentHash ||
+			CurrentTrainedFrameRangesDefaultObject != TrainedDatabaseFrameRangesDefaultObject)
+		{
+			TrainedDatabaseContentHash = CurrentTrainedDatabaseContentHash;
+			TrainedDatabaseFrameRangesDefaultObject = CurrentTrainedFrameRangesDefaultObject;
+			EditorToolKit->EditingAssetWidget->ForceRefresh();
+		}
+		*/
+
+		// TODO: Check Database Skeleton Matches
+
+		bReset = EditorToolKit->bResetRequested;
+
+		// Check if the number of characters has changed
+
+		if (CharacterNum != Controller->TestingSettings->CharacterNum)
+		{
+			CharacterNum = Controller->TestingSettings->CharacterNum;
+
+			for (int32 CharacterIdx = 0; CharacterIdx < MaxCharacterNum; CharacterIdx++)
+			{
+				PreviewScene->SetCharacterVisibility(CharacterIdx, CharacterIdx < CharacterNum);
+			}
+
+			bReset = true;
+		}
+
+		if (bReset)
+		{
+			ActiveBehaviorObject = nullptr;
+		}
+
+		if (BehaviorPreview != ActiveBehaviorObject || BehaviorPreviews.Num() != CharacterNum)
+		{
+			ActiveBehaviorObject = BehaviorPreview;
+
+			const int32 BehaviorNum = BehaviorPreviews.Num();
+			for (int32 BehaviorIdx = 0; BehaviorIdx < BehaviorNum; BehaviorIdx++)
+			{
+				if (BehaviorPreviews[BehaviorIdx] && BehaviorPreviews[BehaviorIdx]->IsValidLowLevel())
+				{
+					BehaviorPreviews[BehaviorIdx]->RemoveFromRoot();
+					BehaviorPreviews[BehaviorIdx] = nullptr;
+				}
+			}
+
+			BehaviorPreviews.SetNumZeroed(CharacterNum);
+
+			for (int32 CharacterIdx = 0; CharacterIdx < CharacterNum; CharacterIdx++)
+			{
+				// Recreate the blueprint UAnimGenBehavior object. This will have the function we need for constructing the controls at the current frame
+				BehaviorPreviews[CharacterIdx] = DuplicateObject(BehaviorPreview, GetTransientPackage());
+				BehaviorPreviews[CharacterIdx]->AddToRoot();
+				BehaviorPreviews[CharacterIdx]->OnPreviewInitialize(Controller, AutoEncoder);
+			}
+
+			ControlObject = UAnimGenControls::MakeControlObject();
+
+			// Force refresh the UI as there is a chance the new behavior class is not compatible
+			EditorToolKit->EditingAssetWidget->ForceRefresh();
+
+			bReset = true;
+		}
+
+		// Create Networks
+
+		if (EncoderInferenceNetwork != AutoEncoder->EncoderNetwork ||
+			DecoderInferenceNetwork != AutoEncoder->DecoderNetwork)
+		{
+			EncoderInferenceNetwork = AutoEncoder->EncoderNetwork;
+			DecoderInferenceNetwork = AutoEncoder->DecoderNetwork;
+			EncoderInference = EncoderInferenceNetwork->GetNetwork()->CreateInferenceObject(MaxCharacterNum);
+			DecoderInference = DecoderInferenceNetwork->GetNetwork()->CreateInferenceObject(MaxCharacterNum);
+		}
+
+		if (ControlEncoderInferenceNetwork != Controller->ControlEncoderNetwork)
+		{
+			ControlEncoderInferenceNetwork = Controller->ControlEncoderNetwork;
+			ControlEncoderInference = ControlEncoderInferenceNetwork->GetNetwork()->CreateInferenceObject(MaxCharacterNum);
+
+			Controller->ControlEncoderSize = ControlEncoderInferenceNetwork->GetSnapshotByteNum() / 1000;
+		}
+
+		if (LOD0InferenceNetwork != Controller->LOD0Network && Controller->LOD0Network)
+		{
+			LOD0InferenceNetwork = Controller->LOD0Network;
+			LOD1InferenceNetwork = Controller->LOD1Network;
+			LOD2InferenceNetwork = Controller->LOD2Network;
+			LOD0Inference = LOD0InferenceNetwork->GetNetwork()->CreateInferenceObject(MaxCharacterNum);
+			LOD1Inference = LOD1InferenceNetwork->GetNetwork()->CreateInferenceObject(MaxCharacterNum);
+			LOD2Inference = LOD2InferenceNetwork->GetNetwork()->CreateInferenceObject(MaxCharacterNum);
+
+			Controller->DenoiserSize =
+				LOD0InferenceNetwork->GetSnapshotByteNum() / 1000 +
+				LOD1InferenceNetwork->GetSnapshotByteNum() / 1000 +
+				LOD2InferenceNetwork->GetSnapshotByteNum() / 1000;
+		}
+
+		// Get Reset Database
+
+		UAnimDatabase* ResetDatabase = Controller->TestingSettings->ResetDatabase.Get() ?
+			Controller->TestingSettings->ResetDatabase.Get() : Controller->TrainingSettings->Database.Get();
+
+		// Compute frame attributes we expect from the reset database
+
+		const int32 DatabaseContentHash = ResetDatabase->GetContentHash();
+
+		if (FrameAttributeEntries != AutoEncoder->TrainedFrameAttributes ||
+			FrameAttributesDatabaseContentHash != DatabaseContentHash)
+		{
+			FScopedSlowTask SlowTask(0.0f, LOCTEXT("RebuildingFrameAttributes", "Rebuilding AutoEncoder Frame Attributes in Database"));
+			SlowTask.MakeDialog();
+
+			FrameAttributeEntries = AutoEncoder->TrainedFrameAttributes;
+			FrameAttributesDatabaseContentHash = DatabaseContentHash;
+
+			const int32 AttributeNum = FrameAttributeEntries.Num();
+			FrameAttributes.SetNum(AttributeNum);
+			for (int32 AttributeIdx = 0; AttributeIdx < AttributeNum; AttributeIdx++)
+			{
+				FrameAttributes[AttributeIdx] = UAnimDatabaseFrameAttributeLibrary::MakeFrameAttributeFromFunction(
+					ResetDatabase, 
+					UAnimDatabaseFrameRangesLibrary::MakeFrameRangesFromDatabase(ResetDatabase),
+					FrameAttributeEntries[AttributeIdx].FrameAttribute);
+			}
+		}
+
+		// Check if we need to manually reset the character state
+
+		if (bReset)
+		{
+			EditorToolKit->bResetRequested = false;
+
+			// Reset the Random State
+
+			RandomState = 0xd5ca438d ^ Controller->TestingSettings->DiffusionSampleSeed;
+
+			// Reset to random frame from tag
+
+			FAnimDatabaseFrameRanges ResetFrameRanges;
+
+			if (Controller->TestingSettings->ResetFrameRanges)
+			{
+				ResetFrameRanges = UAnimDatabaseFrameRangesLibrary::MakeFrameRangesFromFunction(ResetDatabase, UAnimDatabaseFrameRangesLibrary::MakeFrameRangesFromDatabase(ResetDatabase), Controller->TestingSettings->ResetFrameRanges);
+
+				if (!ResetFrameRanges.IsValid())
+				{
+					ResetFrameRanges = UAnimDatabaseFrameRangesLibrary::MakeFrameRangesFromDatabase(ResetDatabase);
+				}
+			}
+			else
+			{
+				ResetFrameRanges = UAnimDatabaseFrameRangesLibrary::MakeFrameRangesFromDatabase(ResetDatabase);
+			}
+
+			if (!PoseState.IsValid())
+			{
+				PoseState = UAnimDatabasePoseStateLibrary::MakePoseState();
+			}
+
+			PoseState.CurrPoseData->Resize(CharacterNum, BoneNum, AutoEncoder->GetAttributeTypes());
+			PoseState.PoseGlobalBoneData->Resize(CharacterNum, BoneNum);
+			PoseVector.SetNumUninitialized({ CharacterNum, PoseVectorSize });
+			UnnormalizedEncodedPoseVector.SetNumUninitialized({ CharacterNum, EncodedPoseVectorSize });
+			EncodedPoseVector.SetNumUninitialized({ CharacterNum, EncodedPoseVectorSize });
+
+			for (int32 CharacterIdx = 0; CharacterIdx < CharacterNum; CharacterIdx++)
+			{
+				// Sample a random sequence and time in the given frame ranges
+				int32 ResetRandomState = 0xfee24eb0 ^ Controller->TestingSettings->ResetSeed;
+
+				int32 SequenceIdx = INDEX_NONE;
+				float SequenceTime = 0.0f;
+				UAnimGenBehaviorLibrary::MakeRandomSequenceIndexAndTimeFromFrameRanges(
+					SequenceIdx,
+					SequenceTime,
+					ResetRandomState,
+					ResetFrameRanges,
+					1.0f / FMath::Max(ResetDatabase->GetFrameRate().AsDecimal(), UE_SMALL_NUMBER));
+
+				// Sample the pose data at the given sequence and time
+				// TODO: Only sample valid bones
+				ResetDatabase->SamplePoseData(
+					PoseState.CurrPoseData->Slice(CharacterIdx, 1),
+					{ SequenceIdx },
+					{ SequenceTime },
+					FrameAttributes);
+
+				// Teleport the character to the origin
+				AnimDatabase::PoseData::Teleport(
+					PoseState.CurrPoseData->RootData.Slice(CharacterIdx, 1),
+					{ FVector::ZeroVector },
+					{ FQuat4f::Identity });
+			}
+
+			*PoseState.PrevPoseData = *PoseState.CurrPoseData;
+			PoseState.DeltaTime = DeltaTime;
+
+			// Compute global bone transforms using forward kinematics
+			AnimDatabase::PoseData::ForwardKinematics(
+				PoseState.PoseGlobalBoneData->View(),
+				PoseState.CurrPoseData->LocalBoneData.ConstView(),
+				PoseState.CurrPoseData->RootData.ConstView(),
+				AutoEncoder->GetBoneParents());
+
+			const int32 LeftToeIndex = AutoEncoder->FindBoneIndex("ball_l");
+			const int32 RightToeIndex = AutoEncoder->FindBoneIndex("ball_r");
+
+			if (LeftToeIndex != INDEX_NONE && RightToeIndex != INDEX_NONE)
+			{
+				ContactState.Resize(CharacterNum, 2);
+				AnimDatabase::PoseData::ContactPinning::ResetContactPinning(
+					ContactState.View(),
+					PoseState.PoseGlobalBoneData->ConstView(),
+					{ LeftToeIndex, RightToeIndex  });
+			}
+
+			// Call the user-implemented blueprint preview scene reset function
+			for (int32 CharacterIdx = 0; CharacterIdx < CharacterNum; CharacterIdx++)
+			{
+				PoseState.PoseIdx = CharacterIdx;
+				BehaviorPreviews[CharacterIdx]->OnPreviewReset(PoseState);
+			}
+
+			// Reset the current encoded pose vector using the auto-encoder EncoderNetwork on the current pose data
+
+			AutoEncoder->ToPoseVectors(PoseVector, PoseState.CurrPoseData->ConstView());
+			AutoEncoder->NormalizePoseVectors(PoseVector);
+			EncoderInference->Evaluate(UnnormalizedEncodedPoseVector, PoseVector);
+			Controller->NormalizeEncodedPoseVectors(EncodedPoseVector, UnnormalizedEncodedPoseVector);
+		}
+
+		// Check for Custom Reset
+
+		for (int32 CharacterIdx = 0; CharacterIdx < CharacterNum; CharacterIdx++)
+		{
+			FAnimGenBehaviorPreviewReset ResetData;
+			if (BehaviorPreviews[CharacterIdx]->OnPreviewCheckReset(ResetData, PoseState))
+			{
+				ResetDatabase->SamplePoseData(
+					PoseState.CurrPoseData->Slice(CharacterIdx, 1),
+					{ ResetData.ResetSequenceIdx },
+					{ ResetData.ResetSequenceTime },
+					FrameAttributes);
+
+				AnimDatabase::PoseData::Teleport(
+					PoseState.CurrPoseData->RootData.Slice(CharacterIdx, 1),
+					{ ResetData.ResetLocation },
+					{ (FQuat4f)ResetData.ResetRotation.Quaternion() });
+
+				AnimDatabase::PoseData::Copy(
+					PoseState.PrevPoseData->Slice(CharacterIdx, 1),
+					PoseState.CurrPoseData->ConstSlice(CharacterIdx, 1));
+				PoseState.DeltaTime = DeltaTime; // TODO: This is for all characters?
+
+				AnimDatabase::PoseData::ForwardKinematics(
+					PoseState.PoseGlobalBoneData->Slice(CharacterIdx, 1),
+					PoseState.CurrPoseData->LocalBoneData.ConstSlice(CharacterIdx, 1),
+					PoseState.CurrPoseData->RootData.ConstSlice(CharacterIdx, 1),
+					AutoEncoder->GetBoneParents());
+
+				const int32 LeftToeIndex = AutoEncoder->FindBoneIndex("ball_l");
+				const int32 RightToeIndex = AutoEncoder->FindBoneIndex("ball_r");
+
+				if (LeftToeIndex != INDEX_NONE && RightToeIndex != INDEX_NONE)
+				{
+					ContactState.Resize(CharacterNum, 2);
+					AnimDatabase::PoseData::ContactPinning::ResetContactPinning(
+						ContactState.View(),
+						PoseState.PoseGlobalBoneData->ConstView(),
+						{ LeftToeIndex, RightToeIndex });
+				}
+
+				// Call the user-implemented blueprint preview scene reset function
+				PoseState.PoseIdx = CharacterIdx;
+				BehaviorPreviews[CharacterIdx]->OnPreviewReset(PoseState);
+
+				// Reset the current encoded pose vector using the auto-encoder EncoderNetwork on the current pose data
+
+				AutoEncoder->ToPoseVectors(PoseVector.Slice(CharacterIdx, 1), PoseState.CurrPoseData->ConstSlice(CharacterIdx, 1));
+				AutoEncoder->NormalizePoseVectors(PoseVector.Slice(CharacterIdx, 1));
+				EncoderInference->Evaluate(UnnormalizedEncodedPoseVector.Slice(CharacterIdx, 1), PoseVector.Slice(CharacterIdx, 1));
+				Controller->NormalizeEncodedPoseVectors(EncodedPoseVector.Slice(CharacterIdx, 1), UnnormalizedEncodedPoseVector.Slice(CharacterIdx, 1));
+			}
+		}
+
+		// Are we playing?
+
+		if (EditorToolKit->ControllerState == EControllerState::Playing ||
+			EditorToolKit->ControllerState == EControllerState::Stepping)
+		{
+			// Capture input when playing (camera cannot move)
+
+			bCaptureInput = true;
+
+			// Evaluate Input
+
+			for (int32 CharacterIdx = 0; CharacterIdx < CharacterNum; CharacterIdx++)
+			{
+				if (BehaviorPreviews[CharacterIdx]->PreviewLeftStick != BehaviorPreviews[CharacterIdx]->PreviewPrevLeftStick)
+				{
+					BehaviorPreviews[CharacterIdx]->OnPreviewInputLeftStickAxis(BehaviorPreviews[CharacterIdx]->PreviewLeftStick);
+					BehaviorPreviews[CharacterIdx]->PreviewPrevLeftStick = BehaviorPreviews[CharacterIdx]->PreviewLeftStick;
+				}
+
+				if (BehaviorPreviews[CharacterIdx]->PreviewRightStick != BehaviorPreviews[CharacterIdx]->PreviewPrevRightStick)
+				{
+					BehaviorPreviews[CharacterIdx]->OnPreviewInputRightStickAxis(BehaviorPreviews[CharacterIdx]->PreviewRightStick);
+					BehaviorPreviews[CharacterIdx]->PreviewPrevRightStick = BehaviorPreviews[CharacterIdx]->PreviewRightStick;
+				}
+
+				if (BehaviorPreviews[CharacterIdx]->PreviewLeftTrigger != BehaviorPreviews[CharacterIdx]->PreviewPrevLeftTrigger)
+				{
+					BehaviorPreviews[CharacterIdx]->OnPreviewInputLeftTriggerAxis(BehaviorPreviews[CharacterIdx]->PreviewLeftTrigger);
+					BehaviorPreviews[CharacterIdx]->PreviewPrevLeftTrigger = BehaviorPreviews[CharacterIdx]->PreviewLeftTrigger;
+				}
+
+				if (BehaviorPreviews[CharacterIdx]->PreviewRightTrigger != BehaviorPreviews[CharacterIdx]->PreviewPrevRightTrigger)
+				{
+					BehaviorPreviews[CharacterIdx]->OnPreviewInputRightTriggerAxis(BehaviorPreviews[CharacterIdx]->PreviewRightTrigger);
+					BehaviorPreviews[CharacterIdx]->PreviewPrevRightTrigger = BehaviorPreviews[CharacterIdx]->PreviewRightTrigger;
+				}
+
+				if (BehaviorPreviews[CharacterIdx]->PreviewMouseGroundLocation != BehaviorPreviews[CharacterIdx]->PreviewPrevMouseGroundLocation)
+				{
+					BehaviorPreviews[CharacterIdx]->OnPreviewInputMouseMoved(BehaviorPreviews[CharacterIdx]->PreviewMouseGroundLocation);
+					BehaviorPreviews[CharacterIdx]->PreviewPrevMouseGroundLocation = BehaviorPreviews[CharacterIdx]->PreviewMouseGroundLocation;
+				}
+			}
+
+			while (KeyEvents.Num())
+			{
+				const TPair<FKey, EInputEvent> KeyEvent = KeyEvents.Pop();
+				const FKey Key = KeyEvent.Key;
+				const EInputEvent Event = KeyEvent.Value;
+
+				for (int32 CharacterIdx = 0; CharacterIdx < CharacterNum; CharacterIdx++)
+				{
+					if (Key == EKeys::LeftMouseButton && Event == EInputEvent::IE_Pressed)
+					{
+						BehaviorPreviews[CharacterIdx]->OnPreviewInputLeftMousePressed(BehaviorPreviews[CharacterIdx]->PreviewMouseGroundLocation);
+					}
+					else if (Key == EKeys::LeftMouseButton && Event == EInputEvent::IE_Released)
+					{
+						BehaviorPreviews[CharacterIdx]->OnPreviewInputLeftMouseReleased(BehaviorPreviews[CharacterIdx]->PreviewMouseGroundLocation);
+					}
+					else if (Key == EKeys::RightMouseButton && Event == EInputEvent::IE_Pressed)
+					{
+						BehaviorPreviews[CharacterIdx]->OnPreviewInputRightMousePressed(BehaviorPreviews[CharacterIdx]->PreviewMouseGroundLocation);
+					}
+					else if (Key == EKeys::RightMouseButton && Event == EInputEvent::IE_Released)
+					{
+						BehaviorPreviews[CharacterIdx]->OnPreviewInputRightMouseReleased(BehaviorPreviews[CharacterIdx]->PreviewMouseGroundLocation);
+					}
+					else if (Event == EInputEvent::IE_Pressed)
+					{
+						BehaviorPreviews[CharacterIdx]->OnPreviewInputKeyPressed(Key);
+					}
+					else if (Event == EInputEvent::IE_Released)
+					{
+						BehaviorPreviews[CharacterIdx]->OnPreviewInputKeyReleased(Key);
+					}
+				}
+			}
+
+			// Update Camera
+
+			if (Controller->TestingSettings->bFollowCam && CharacterNum > 0)
+			{
+				ViewportClient->SetViewRotation(FRotator(-BehaviorPreviews[0]->PreviewCameraPitch, BehaviorPreviews[0]->PreviewCameraYaw, 0.0f));
+				ViewportClient->SetViewLocationForOrbiting(
+					PoseState.CurrPoseData->RootData.RootLocations[0] + FVector(0.0f, 0.0f, BehaviorPreviews[0]->PreviewCameraHeight),
+					BehaviorPreviews[0]->PreviewCameraDistance);
+			}
+
+			for (int32 CharacterIdx = 0; CharacterIdx < CharacterNum; CharacterIdx++)
+			{
+				BehaviorPreviews[CharacterIdx]->PreviewCameraLocation = ViewportClient->GetViewLocation();
+				BehaviorPreviews[CharacterIdx]->PreviewCameraRotation = ViewportClient->GetViewRotation();
+			}
+
+			// Update Preview Scene
+
+			for (int32 CharacterIdx = 0; CharacterIdx < CharacterNum; CharacterIdx++)
+			{
+				PoseState.PoseIdx = CharacterIdx;
+				BehaviorPreviews[CharacterIdx]->OnPreviewUpdate(PoseState, DeltaTime);
+			}
+
+			// Record Previous Frame
+
+			*PoseState.PrevPoseData = *PoseState.CurrPoseData;
+			PoseState.DeltaTime = DeltaTime;
+
+			// Evaluate Controls. Here we call the blueprint implemented MakeControl function using the current character state.
+
+			const int32 ControlVectorSize = Controller->ControlSchema.ObservationSchema->Schema.GetObservationVectorSize(Controller->ControlSchemaElement.SchemaElement);
+			const int32 ControlDistributionVectorSize = Controller->ControlSchema.ObservationSchema->Schema.GetObservationDistributionVectorSize(Controller->ControlSchemaElement.SchemaElement);
+
+			ControlVector.SetNumUninitialized({ CharacterNum, ControlVectorSize });
+			ControlObject.ObservationObject->Object.Reset();
+			ControlObjectElements.SetNum(CharacterNum);
+
+			for (int32 CharacterIdx = 0; CharacterIdx < CharacterNum; CharacterIdx++)
+			{
+				PoseState.PoseIdx = CharacterIdx;
+				ControlObjectElements[CharacterIdx] = BehaviorPreviews[CharacterIdx]->MakeControl(ControlObject, PoseState);
+
+				// Validate that the return controls match the schema we provided
+
+				if (!UAnimGenControls::ValidateControlObjectMatchesSchema(
+					Controller->ControlSchema,
+					Controller->ControlSchemaElement,
+					ControlObject,
+					ControlObjectElements[CharacterIdx]))
+				{
+					return;
+				}
+
+				Learning::Observation::SetVectorFromObject(
+					ControlVector[CharacterIdx],
+					Controller->ControlSchema.ObservationSchema->Schema,
+					Controller->ControlSchemaElement.SchemaElement,
+					ControlObject.ObservationObject->Object,
+					ControlObjectElements[CharacterIdx].ObjectElement);
+			}
+
+			// Clamp Control Vectors
+
+			if (Controller->ControlVectorDistributionMins.Num() == ControlDistributionVectorSize &&
+				Controller->ControlVectorDistributionMaxs.Num() == ControlDistributionVectorSize)
+			{
+				for (int32 CharacterIdx = 0; CharacterIdx < CharacterNum; CharacterIdx++)
+				{
+					Learning::Observation::ClampToObservationDistributionMinMax(
+						ControlVector[CharacterIdx],
+						Controller->ControlVectorDistributionMins,
+						Controller->ControlVectorDistributionMaxs,
+						Controller->ControlSchema.ObservationSchema->Schema,
+						Controller->ControlSchemaElement.SchemaElement);
+				}
+			}
+
+			// Evaluate Control Encoder
+
+			const int32 EncodedControlVectorSize = Controller->EncodedControlVectorSize;
+			ControllerOutputVector.SetNumUninitialized({ CharacterNum, EncodedControlVectorSize });
+
+			if (!ControlEncoderInference)
+			{
+				return;
+			}
+
+			if (ControllerOutputVector.Num<1>() != ControlEncoderInference->GetOutputSize() ||
+				ControlVector.Num<1>() != ControlEncoderInference->GetInputSize())
+			{
+				ViewportClient->SetWarningMessage(LOCTEXT("NetworkSizeWarning", "Warning: Control Encoder Network sizes don't match. Is network asset out-of-date?"));
+				Learning::Array::Zero(ControllerOutputVector);
+			}
+			else
+			{
+				ControlEncoderPerfCounter.Begin();
+				if (EncodedControlVectorSize != 0) { ControlEncoderInference->Evaluate(ControllerOutputVector, ControlVector); }
+				ControlEncoderPerfCounter.End();
+
+				Controller->ControlEncoderInferenceTime = FMath::RoundToInt(FPlatformTime::ToMilliseconds64(ControlEncoderPerfCounter.GetAverage()) * 1000.0);
+			}
+
+			// Evaluate Denoiser
+
+			const float ControllerDeltaTime = 1.0f / FMath::Max(Controller->FrameRate.AsDecimal(), UE_SMALL_NUMBER);
+			const float SubsampleRate = FMath::Clamp(FMath::Max(DeltaTime, ControllerDeltaTime) / ControllerDeltaTime, 1, 8);
+
+			NextEncodedPoseVector = EncodedPoseVector;
+
+			DenoiserPerfCounter.Begin();
+
+			DenoiserInputVector.SetNumUninitialized({ CharacterNum, EncodedPoseVectorSize * 2 + 1 + EncodedControlVectorSize });
+
+			check(LOD0Inference->GetInputSize() == LOD1Inference->GetInputSize() && LOD0Inference->GetInputSize() == LOD2Inference->GetInputSize());
+			check(LOD0Inference->GetOutputSize() == LOD1Inference->GetOutputSize() && LOD0Inference->GetOutputSize() == LOD2Inference->GetOutputSize());
+
+			if (DenoiserInputVector.Num<1>() != LOD0Inference->GetInputSize() || EncodedPoseVectorSize != LOD0Inference->GetOutputSize())
+			{
+				ViewportClient->SetWarningMessage(LOCTEXT("NetworkSizeWarning", "Warning: Denoiser Network sizes don't match. Is network asset out-of-date?"));
+				Learning::Array::Zero(NextEncodedPoseVector);
+			}
+			else
+			{
+				for (int32 CharacterIdx = 0; CharacterIdx < CharacterNum; CharacterIdx++)
+				{
+					Learning::Random::SampleClippedGaussianArray(NextEncodedPoseVector[CharacterIdx], RandomState, 0.0f, Controller->TestingSettings->NoiseScale, 3.0f);
+					if (!Controller->NormalizedPoseStds.IsEmpty())
+					{
+						for (int32 Idx = 0; Idx < EncodedPoseVectorSize; Idx++)
+						{
+							NextEncodedPoseVector[CharacterIdx][Idx] *= Controller->NormalizedPoseStds[Idx];
+						}
+					}
+
+					Learning::Array::Copy(DenoiserInputVector[CharacterIdx].Slice(0, EncodedPoseVectorSize), EncodedPoseVector[CharacterIdx]);
+					Learning::Array::Copy(DenoiserInputVector[CharacterIdx].Slice(EncodedPoseVectorSize, EncodedPoseVectorSize), NextEncodedPoseVector[CharacterIdx]);
+					Learning::Array::Set(DenoiserInputVector[CharacterIdx].Slice(EncodedPoseVectorSize * 2, 1), (SubsampleRate - 1.0f) / 7.0f);
+					Learning::Array::Copy(DenoiserInputVector[CharacterIdx].Slice(EncodedPoseVectorSize * 2 + 1, EncodedControlVectorSize), ControllerOutputVector[CharacterIdx]);
+				}
+
+				switch (Controller->TestingSettings->LODLevel)
+				{
+				default:
+				case 0: LOD0Inference->Evaluate(NextEncodedPoseVector, DenoiserInputVector); break;
+				case 1: LOD1Inference->Evaluate(NextEncodedPoseVector, DenoiserInputVector); break;
+				case 2: LOD2Inference->Evaluate(NextEncodedPoseVector, DenoiserInputVector); break;
+				}
+
+				Controller->ClampNormalizedPoseVectorsInplace(NextEncodedPoseVector);
+			}
+
+			DenoiserPerfCounter.End();
+
+			Controller->DenoiserInferenceTime = FMath::RoundToInt(FPlatformTime::ToMilliseconds64(DenoiserPerfCounter.GetAverage()) * 1000.0);
+
+			// Interpolate / Extrapolate using Delta Time Ratio
+
+			if (DeltaTime < ControllerDeltaTime)
+			{
+				AnimDatabase::Math::LerpToTargetInplace(
+					EncodedPoseVector,
+					NextEncodedPoseVector,
+					DeltaTime / ControllerDeltaTime);
+			}
+			else
+			{
+				Learning::Array::Copy(EncodedPoseVector, NextEncodedPoseVector);
+			}
+
+			Controller->ClampNormalizedPoseVectorsInplace(EncodedPoseVector);
+
+			// Evaluate Decoder
+
+			// Using the Decoder we can map from the EncodedPoseVector to the PoseVector, and then from the PoseVector we can get the PoseData
+
+			Controller->DenormalizeEncodedPoseVectors(UnnormalizedEncodedPoseVector, EncodedPoseVector);
+			PoseVector.SetNumUninitialized({ CharacterNum, PoseVectorSize });
+			DecoderInference->Evaluate(PoseVector, UnnormalizedEncodedPoseVector);
+
+			// Apply some smoothing based on mixing predicted velocities with predicted rotations and also account for DeltaTime variations
+
+			OutputPoseData.Resize(CharacterNum, BoneNum, AutoEncoder->GetAttributeTypes());
+
+			AutoEncoder->DenormalizePoseVectors(PoseVector);
+
+			AutoEncoder->FromPoseVectors(
+				OutputPoseData.View(),
+				PoseVector,
+				PoseState.CurrPoseData->RootData.RootLocations,
+				PoseState.CurrPoseData->RootData.RootRotations);
+
+			// Update Attributes
+			
+			PoseState.CurrPoseData->AttributeData = OutputPoseData.AttributeData;
+
+			// Update Pose
+
+			AnimDatabase::PoseData::DampInplaceUsingAdaptiveVelocityIntegrationMix(
+				PoseState.CurrPoseData->LocalBoneData.View(),
+				OutputPoseData.LocalBoneData.ConstView(),
+				DeltaTime,
+				Controller->TestingSettings->MinPoseSmoothingTime,
+				Controller->TestingSettings->MaxPoseSmoothingTime);
+
+			// Integrate Root
+
+			PoseState.CurrPoseData->RootData = OutputPoseData.RootData;
+			for (int32 CharacterIdx = 0; CharacterIdx < CharacterNum; CharacterIdx++)
+			{
+				PoseState.CurrPoseData->RootData.RootLocations[CharacterIdx] = DeltaTime * (FVector)OutputPoseData.RootData.RootLinearVelocities[CharacterIdx] + OutputPoseData.RootData.RootLocations[CharacterIdx];
+				PoseState.CurrPoseData->RootData.RootRotations[CharacterIdx] = FQuat4f::MakeFromRotationVector(DeltaTime * OutputPoseData.RootData.RootAngularVelocities[CharacterIdx]) * OutputPoseData.RootData.RootRotations[CharacterIdx];
+			}
+
+			PoseState.PoseGlobalBoneData->Resize(CharacterNum, BoneNum);
+			AnimDatabase::PoseData::ForwardKinematics(PoseState.PoseGlobalBoneData->View(), PoseState.CurrPoseData->LocalBoneData.ConstView(), PoseState.CurrPoseData->RootData.ConstView(), AutoEncoder->GetBoneParents());
+
+			AdjustedPoseData = *PoseState.CurrPoseData;
+
+			const int32 HipBoneIndex = AutoEncoder->FindBoneIndex("pelvis");
+
+			if (Controller->TestingSettings->bEnableFootLocking && HipBoneIndex != INDEX_NONE)
+			{
+				for (int32 CharacterIdx = 0; CharacterIdx < CharacterNum; CharacterIdx++)
+				{
+					AdjustedPoseData.LocalBoneData.BoneLocations[CharacterIdx][HipBoneIndex].Z += Controller->TestingSettings->FootLockingPelvisHeightAdjustment;
+				}
+			}
+
+			AdjustedGlobalBoneData.Resize(CharacterNum, BoneNum);
+			AnimDatabase::PoseData::ForwardKinematics(AdjustedGlobalBoneData.View(), AdjustedPoseData.LocalBoneData.ConstView(), AdjustedPoseData.RootData.ConstView(), AutoEncoder->BoneParents);
+
+			///////////////////
+
+			const int32 LeftContactAttribute = AutoEncoder->FindAttributeIndex(TEXT("contact_l"));
+			const int32 LeftLegBoneIndex = AutoEncoder->FindBoneIndex("thigh_l");
+			const int32 LeftKneeBoneIndex = AutoEncoder->FindBoneIndex("calf_l");
+			const int32 LeftAnkleBoneIndex = AutoEncoder->FindBoneIndex("foot_l");
+			const int32 LeftToeBoneIndex = AutoEncoder->FindBoneIndex("ball_l");
+
+			const int32 RightContactAttribute = AutoEncoder->FindAttributeIndex(TEXT("contact_r"));
+			const int32 RightLegBoneIndex = AutoEncoder->FindBoneIndex("thigh_r");
+			const int32 RightKneeBoneIndex = AutoEncoder->FindBoneIndex("calf_r");
+			const int32 RightAnkleBoneIndex = AutoEncoder->FindBoneIndex("foot_r");
+			const int32 RightToeBoneIndex = AutoEncoder->FindBoneIndex("ball_r");
+
+			if (Controller->TestingSettings->bEnableFootLocking &&
+				HipBoneIndex != INDEX_NONE &&
+				LeftContactAttribute != INDEX_NONE &&
+				LeftLegBoneIndex != INDEX_NONE &&
+				LeftKneeBoneIndex != INDEX_NONE &&
+				LeftAnkleBoneIndex != INDEX_NONE &&
+				LeftToeBoneIndex != INDEX_NONE &&
+				RightContactAttribute != INDEX_NONE &&
+				RightLegBoneIndex != INDEX_NONE &&
+				RightKneeBoneIndex != INDEX_NONE &&
+				RightAnkleBoneIndex != INDEX_NONE &&
+				RightToeBoneIndex != INDEX_NONE)
+			{
+				AnimDatabase::PoseData::ContactPinning::FContactPinningSettings ContactSettings;
+				ContactSettings.BlendTime = Controller->TestingSettings->FootLockingBlendTime;
+				ContactSettings.LockRadius = Controller->TestingSettings->FootLockingLockRadius;
+				ContactSettings.UnlockRadius = Controller->TestingSettings->FootLockingUnlockRadius;
+				ContactSettings.ContactHeight = AutoEncoder->TrainedSkeleton->GetReferenceSkeleton().GetBoneAbsoluteTransform(LeftToeBoneIndex).GetLocation().Z;
+				ContactSettings.MinToeHeight = AutoEncoder->TrainedSkeleton->GetReferenceSkeleton().GetBoneAbsoluteTransform(LeftToeBoneIndex).GetLocation().Z;
+				ContactSettings.MinAnkleHeight = AutoEncoder->TrainedSkeleton->GetReferenceSkeleton().GetBoneAbsoluteTransform(LeftAnkleBoneIndex).GetLocation().Z;
+				ContactSettings.ToeLength = Controller->TestingSettings->FootLockingToeLength;
+
+				AnimDatabase::PoseData::ContactPinning::UpdateContactPinning(
+					ContactState.View(),
+					AdjustedGlobalBoneData.View(),
+					AdjustedPoseData.LocalBoneData.View(),
+					AdjustedPoseData.RootData.ConstView(),
+					AdjustedPoseData.AttributeData.ConstView(),
+					AutoEncoder->GetBoneParents(),
+					{ LeftContactAttribute, RightContactAttribute },
+					{ HipBoneIndex, HipBoneIndex },
+					{ LeftLegBoneIndex, RightLegBoneIndex },
+					{ LeftKneeBoneIndex, RightKneeBoneIndex },
+					{ LeftAnkleBoneIndex, RightAnkleBoneIndex },
+					{ LeftToeBoneIndex, RightToeBoneIndex },
+					{ FVector3f(0.0f, 0.0f, +1.0f), FVector3f(0.0f, 0.0f, +1.0f) },
+					{ FVector3f(-1.0f, 0.0f, 0.0f), FVector3f(-1.0f, 0.0f, 0.0f) },
+					DeltaTime,
+					ContactSettings);
+			}
+
+			///////////////////
+
+			// Forward Kinematics to get the global transforms
+
+			AdjustedGlobalBoneData.Resize(CharacterNum, BoneNum);
+			AnimDatabase::PoseData::ForwardKinematics(AdjustedGlobalBoneData.View(), AdjustedPoseData.LocalBoneData.ConstView(), AdjustedPoseData.RootData.ConstView(), AutoEncoder->GetBoneParents());
+
+			///////////////////
+
+			// TODO: Why does there seem to be some error here where the root moves before the bones are updated?
+			//AnimGenViewportClient->GetPreviewScene()->GetCharacterActor()->SetActorTransform(RootTransform);
+
+			// Push the global transforms onto the character
+
+			for (int32 CharacterIdx = 0; CharacterIdx < CharacterNum; CharacterIdx++)
+			{
+				if (PreviewScene->GetPoseableMeshComponent(CharacterIdx)->GetSkinnedAsset())
+				{
+					for (int32 BoneIdx = 0; BoneIdx < BoneNum; BoneIdx++)
+					{
+						const FTransform BoneTransform = FTransform(
+							((FQuat)AdjustedGlobalBoneData.BoneRotations[CharacterIdx][BoneIdx]).GetNormalized(),
+							(FVector)AdjustedGlobalBoneData.BoneLocations[CharacterIdx][BoneIdx],
+							(FVector)AdjustedGlobalBoneData.BoneScales[CharacterIdx][BoneIdx]);
+
+						PreviewScene->GetPoseableMeshComponent(CharacterIdx)->SetBoneTransformByName(AutoEncoder->GetBoneName(BoneIdx), BoneTransform, EBoneSpaces::WorldSpace);
+					}
+
+					PreviewScene->GetPoseableMeshComponent(CharacterIdx)->RefreshBoneTransforms();
+				}
+			}
+
+			// Update UI Control Objects
+
+			if (ControlObject.ObservationObject.IsValid() && CharacterNum > 0)
+			{
+				TArray<TSharedPtr<FControlObjectTreeElement>>& ControlObjectTreeElements = EditorToolKit->ControlObjectWidget->GetTreeElementsRef();
+
+				bool bRefreshControlUI = ControlObjectTreeElements.Num() != CharacterNum;
+
+				ControlObjectTreeElements.SetNum(CharacterNum);
+
+				for (int32 CharacterIdx = 0; CharacterIdx < CharacterNum; CharacterIdx++)
+				{
+					if (!ControlObjectTreeElements[CharacterIdx].IsValid())
+					{
+						ControlObjectTreeElements[CharacterIdx] = MakeShared<FControlObjectTreeElement>();
+						bRefreshControlUI = true;
+					}
+
+					bRefreshControlUI |= Controller::Private::UpdateControlObjectTreeElement(
+						ControlObjectTreeElements[CharacterIdx],
+						ControlObject.ObservationObject->Object,
+						ControlObjectElements[CharacterIdx].ObjectElement);
+
+					ControlObjectTreeElements[CharacterIdx]->Color = CharacterColors[CharacterIdx];
+				}
+
+				if (bRefreshControlUI)
+				{
+					EditorToolKit->ControlObjectWidget->RefreshTreeElements();
+				}
+			}
+			else
+			{
+				TArray<TSharedPtr<FControlObjectTreeElement>>& ControlObjectTreeElements = EditorToolKit->ControlObjectWidget->GetTreeElementsRef();
+				if (!ControlObjectTreeElements.IsEmpty())
+				{
+					ControlObjectTreeElements.Empty();
+					EditorToolKit->ControlObjectWidget->RefreshTreeElements();
+				}
+			}
+
+			if (EditorToolKit->ControllerState == EControllerState::Stepping)
+			{
+				EditorToolKit->ControllerState = EControllerState::Paused;
+			}
+		}
+		else
+		{
+			bCaptureInput = false;
+		}
+	}
+
+	void FControllerMode::Tick(FEditorViewportClient* EditorViewportClient, float DeltaTime)
+	{
+		FEdMode::Tick(EditorViewportClient, DeltaTime);
+
+		AnimDatabase::Editor::FViewportClient* ViewportClient = static_cast<AnimDatabase::Editor::FViewportClient*>(EditorViewportClient);
+		check(ViewportClient);
+
+		TSharedPtr<FControllerToolkit> EditorToolKit = StaticCastSharedPtr<FControllerToolkit>(ViewportClient->GetAssetEditorToolkit().Pin());
+		if (!EditorToolKit.IsValid()) { return; }
+
+		WeakToolkit = EditorToolKit.ToWeakPtr();
+
+		if (EditorToolKit->ControllerState == EControllerState::Stopped)
+		{
+			TickTraining(EditorViewportClient, DeltaTime);
+			bReset = true;
+		}
+		else
+		{
+			TickTesting(EditorViewportClient, DeltaTime);
+		}
+	}
+
+	void FControllerMode::RenderTraining(const FSceneView* View, FViewport* Viewport, FPrimitiveDrawInterface* PDI)
+	{
+		if (TSharedPtr<FControllerToolkit> ToolKit = WeakToolkit.Pin())
+		{
+			// Load Database
+
+			UAnimDatabase* TrainingDatabase = ToolKit->Controller->TrainingSettings->Database.Get();
+			if (!TrainingDatabase) { return; }
+
+			// Load AutoEncoder
+
+			UAnimGenAutoEncoder* TrainingAutoEncoder = ToolKit->Controller->TrainingSettings->AutoEncoder.Get();
+			if (!TrainingAutoEncoder || !TrainingAutoEncoder->IsValid()) { return; }
+
+			// Load Behavior
+
+			TObjectPtr<UAnimGenBehavior> TrainingBehavior = ToolKit->Controller->TrainingSettings->Behavior;
+			if (!TrainingBehavior) { return; }
+
+			// Draw Skeleton
+
+			if (ToolKit->Controller->ViewportSettings->bDrawSkeleton)
+			{
+				// Find bone indices to render
+
+				const int32 BoneNum = TrainingAutoEncoder->GetBoneNum();
+
+				TArray<int32> BoneIndices;
+
+				if (ToolKit->Controller->ViewportSettings->bDrawRequiredBonesOnly)
+				{
+					BoneIndices = TrainingAutoEncoder->AutoEncodedRequiredBoneIndices;
+				}
+				else
+				{
+					BoneIndices.SetNumUninitialized(BoneNum);
+					for (int32 BoneIdx = 0; BoneIdx < BoneNum; BoneIdx++)
+					{
+						BoneIndices[BoneIdx] = BoneIdx;
+					}
+				}
+
+				const FLinearColor DefaultBoneColor = GetDefault<UPersonaOptions>()->DefaultBoneColor;
+
+				for (int32 CharacterIdx = 0; CharacterIdx < CharacterNum; CharacterIdx++)
+				{
+					// Draw Skeleton
+
+					FLinearColor BoneColor = DefaultBoneColor;
+					BoneColor.A = ToolKit->Controller->ViewportSettings->DrawSkeletonOpacity;
+
+					FDrawDebugSkeletonSettings SkeletonSettings;
+					SkeletonSettings.bDrawSimpleSkeleton = ToolKit->Controller->ViewportSettings->bDrawSimpleSkeleton;
+					SkeletonSettings.BoneRadius = ToolKit->Controller->ViewportSettings->DrawSkeletonScale;
+
+					UDrawDebugLibrary::DrawDebugSkeletonArrayView(
+						FDebugDrawer::MakeDebugDrawer(PDI),
+						PoseState.CurrPoseData->RootData.RootLocations[CharacterIdx],
+						PoseState.CurrPoseData->RootData.RootRotations[CharacterIdx],
+						PoseState.PoseGlobalBoneData->BoneLocations[CharacterIdx],
+						PoseState.PoseGlobalBoneData->BoneRotations[CharacterIdx],
+						BoneIndices,
+						ToolKit->Controller->AutoEncoder->GetBoneParents(),
+						BoneColor,
+						false,
+						SkeletonSettings);
+
+					// Draw Velocities
+
+					FLinearColor VelocityBoneColor = BoneColor;
+					VelocityBoneColor.A = ToolKit->Controller->ViewportSettings->DrawLinearVelocitiesOpacity;
+
+					if (ToolKit->Controller->ViewportSettings->bDrawLinearVelocities)
+					{
+						UDrawDebugLibrary::DrawDebugBoneVelocitiesArrayView(
+							FDebugDrawer::MakeDebugDrawer(PDI),
+							PoseState.PoseGlobalBoneData->BoneLocations[CharacterIdx],
+							PoseState.PoseGlobalBoneData->BoneLinearVelocities[CharacterIdx],
+							BoneIndices,
+							VelocityBoneColor,
+							ToolKit->Controller->ViewportSettings->DrawSkeletonScale * 0.5f,
+							true,
+							ToolKit->Controller->ViewportSettings->DrawLinearVelocitiesScale);
+					}
+
+					if (ToolKit->Controller->ViewportSettings->bDrawAngularVelocities)
+					{
+						UDrawDebugLibrary::DrawDebugBoneVelocitiesArrayView(
+							FDebugDrawer::MakeDebugDrawer(PDI),
+							PoseState.PoseGlobalBoneData->BoneLocations[CharacterIdx],
+							PoseState.PoseGlobalBoneData->BoneAngularVelocities[CharacterIdx],
+							BoneIndices,
+							VelocityBoneColor,
+							ToolKit->Controller->ViewportSettings->DrawSkeletonScale * 0.5f,
+							true,
+							ToolKit->Controller->ViewportSettings->DrawAngularVelocitiesScale);
+					}
+
+					// Draw Transforms
+
+					if (ToolKit->Controller->ViewportSettings->bDrawBoneTransforms)
+					{
+						UDrawDebugLibrary::DrawDebugTransformsArrayView(
+							FDebugDrawer::MakeDebugDrawer(PDI),
+							PoseState.PoseGlobalBoneData->BoneLocations[CharacterIdx],
+							PoseState.PoseGlobalBoneData->BoneRotations[CharacterIdx],
+							PoseState.PoseGlobalBoneData->BoneScales[CharacterIdx],
+							UDrawDebugLibrary::MakeDrawDebugLineStyleFromColor(BoneColor),
+							false,
+							ToolKit->Controller->ViewportSettings->DrawSkeletonScale * 5.0f);
+					}
+				}
+			}
+
+			// Draw Range Identifiers
+
+			if (ToolKit->Controller->ViewportSettings->bDrawRangeIdentifier)
+			{
+				for (int32 CharacterIdx = 0; CharacterIdx < CharacterNum; CharacterIdx++)
+				{
+					FLinearColor IdentifierColor = ToolKit->Controller->TrainingSettings->QueryEntries[ActiveRanges[CharacterIdx]]->Color;
+					IdentifierColor.A = ToolKit->Controller->ViewportSettings->DrawRangeIdentifierOpacity;
+
+					const FVector IdentifierLocation = (FVector)PoseState.CurrPoseData->RootData.RootLocations[CharacterIdx] +
+						(FVector)PoseState.CurrPoseData->RootData.RootRotations[CharacterIdx].RotateVector(FVector3f(0.0f, 0.0f, ToolKit->Controller->ViewportSettings->DrawRangeIdentifierHeight));
+
+					FDrawDebugLineStyle LineStyle;
+					LineStyle.Color = IdentifierColor;
+					LineStyle.Thickness = ToolKit->Controller->ViewportSettings->DrawRangeIdentifierThickness;
+
+					UDrawDebugLibrary::DrawDebugSimpleSphere(
+						FDebugDrawer::MakeDebugDrawer(PDI),
+						IdentifierLocation,
+						FRotator::ZeroRotator,
+						LineStyle,
+						true,
+						ToolKit->Controller->ViewportSettings->DrawRangeIdentifierRadius);
+				}
+			}
+
+			// Draw Skeleton Roots
+
+			if (ToolKit->Controller->ViewportSettings->bDrawRoot && CharacterNum > 0)
+			{
+				FDrawDebugLineStyle LineStyle;
+				LineStyle.Color = FLinearColor::White;
+				LineStyle.Color.A = ToolKit->Controller->ViewportSettings->DrawRootOpacity;
+				LineStyle.Thickness = ToolKit->Controller->ViewportSettings->DrawSkeletonScale;
+
+				UDrawDebugLibrary::DrawDebugRotationsQuatArrayView(
+					FDebugDrawer::MakeDebugDrawer(PDI),
+					PoseState.CurrPoseData->RootData.RootLocations,
+					PoseState.CurrPoseData->RootData.RootRotations,
+					LineStyle,
+					true,
+					ToolKit->Controller->ViewportSettings->DrawRootScale);
+			}
+
+			// Draw Origin
+
+			if (ToolKit->Controller->ViewportSettings->bDrawOrigin)
+			{
+				FDrawDebugLineStyle LineStyle;
+				LineStyle.Color = FLinearColor::White;
+				LineStyle.Color.A = ToolKit->Controller->ViewportSettings->DrawOriginOpacity;
+				LineStyle.Thickness = ToolKit->Controller->ViewportSettings->DrawOriginLineThickness;
+
+				UDrawDebugLibrary::DrawDebugTransform(
+					FDebugDrawer::MakeDebugDrawer(PDI),
+					FTransform::Identity,
+					LineStyle,
+					true,
+					ToolKit->Controller->ViewportSettings->DrawOriginScale);
+			}
+
+			// Draw Trajectory
+
+			if (ToolKit->Controller->ViewportSettings->bDrawTrajectories &&
+				(!ToolKit->Controller->ViewportSettings->bRangesStartAtOrigin ||
+					!ToolKit->Controller->ViewportSettings->bDisableRootTranslation))
+			{
+				FVector ForwardVector = FVector::ForwardVector;
+				switch (TrainingAutoEncoder->TrainedSkeleton->GetPreviewForwardAxis())
+				{
+				case EAxis::X: ForwardVector = FVector::XAxisVector; break;
+				case EAxis::Y: ForwardVector = FVector::YAxisVector; break;
+				case EAxis::Z: ForwardVector = FVector::ZAxisVector; break;
+				}
+
+				for (int32 CharacterIdx = 0; CharacterIdx < CharacterNum; CharacterIdx++)
+				{
+					const FVector RootOffset = Controller::Private::GetRangeOffset(
+						CharacterIdx, CharacterNum,
+						ToolKit->Controller->ViewportSettings->RangeOffsetSpacing,
+						ToolKit->Controller->ViewportSettings->bOffsetRanges,
+						ToolKit->Controller->ViewportSettings->bOffsetRangesInGrid);
+
+					UDrawDebugLibrary::DrawDebugRangeTrajectoryArrayView(
+						FDebugDrawer::MakeDebugDrawer(PDI),
+						RangeTrajectoryLocations[ActiveRanges[CharacterIdx]],
+						RangeTrajectoryRotations[ActiveRanges[CharacterIdx]],
+						true,
+						ForwardVector,
+						ToolKit->Controller->ViewportSettings->bDrawTrajectoryOrientations,
+						ToolKit->Controller->ViewportSettings->bRangesStartAtOrigin,
+						RootOffset);
+				}
+			}
+
+			// Debug Draw Controls
+
+			if (ToolKit->Controller->ViewportSettings->bDrawDebugTraining && TrainingBehavior)
+			{
+				for (int32 CharacterIdx = 0; CharacterIdx < CharacterNum; CharacterIdx++)
+				{
+					FLinearColor IdentifierColor = ToolKit->Controller->TrainingSettings->QueryEntries[ActiveRanges[CharacterIdx]]->Color;
+					IdentifierColor.A = ToolKit->Controller->ViewportSettings->DrawRangeIdentifierOpacity;
+
+					const int32 FlatIndex = ToolKit->Controller->TrainingSettings->ControlSets[CharacterControlSetIndices[CharacterIdx]].Data->FrameRangeSet.GetEntryRangeOffset(
+						CharacterControlSetEntries[CharacterIdx], CharacterControlSetRanges[CharacterIdx]) + CharacterControlSetFrames[CharacterIdx];
+
+					PoseState.PoseIdx = CharacterIdx;
+
+					TrainingBehavior->DrawDebugControl(
+						FDebugDrawer::MakeDebugDrawer(PDI),
+						ToolKit->Controller->TrainingSettings->ControlObject,
+						ToolKit->Controller->TrainingSettings->ControlSets[CharacterControlSetIndices[CharacterIdx]].Data->Controls[FlatIndex],
+						PoseState,
+						ToolKit->Controller->TrainingSettings->Database.Get(),
+						ToolKit->Controller->TrainingSettings->AutoEncoder.Get(),
+						CharacterSequenceIndices[CharacterIdx],
+						CharacterRangeTimes[CharacterIdx],
+						CharacterRangeStarts[CharacterIdx],
+						CharacterRangeLengths[CharacterIdx],
+						IdentifierColor);
+				}
+			}
+		}
+	}
+
+	void FControllerMode::RenderTesting(const FSceneView* View, FViewport* Viewport, FPrimitiveDrawInterface* PDI)
+	{
+		if (TSharedPtr<FControllerToolkit> ToolKit = WeakToolkit.Pin())
+		{
+			if (!ToolKit->Controller->IsValid()) { return; }
+			if (!ToolKit->Controller->AutoEncoder->IsValid()) { return; }
+			if (!ToolKit->Controller->TestingSettings->BehaviorPreview) { return; }
+
+			// Draw Skeleton
+
+			if (ToolKit->Controller->ViewportSettings->bDrawSkeleton)
+			{
+				// Find bone indices to render
+
+				const int32 BoneNum = ToolKit->Controller->AutoEncoder->GetBoneNum();
+
+				TArray<int32> BoneIndices;
+
+				if (ToolKit->Controller->ViewportSettings->bDrawRequiredBonesOnly)
+				{
+					BoneIndices = ToolKit->Controller->AutoEncoder->AutoEncodedRequiredBoneIndices;
+				}
+				else
+				{
+					BoneIndices.SetNumUninitialized(BoneNum);
+					for (int32 BoneIdx = 0; BoneIdx < BoneNum; BoneIdx++)
+					{
+						BoneIndices[BoneIdx] = BoneIdx;
+					}
+				}
+
+				const FLinearColor DefaultBoneColor = GetDefault<UPersonaOptions>()->DefaultBoneColor;
+
+				for (int32 CharacterIdx = 0; CharacterIdx < CharacterNum; CharacterIdx++)
+				{
+					// Draw Skeleton
+
+					FLinearColor BoneColor = DefaultBoneColor;
+					BoneColor.A = ToolKit->Controller->ViewportSettings->DrawSkeletonOpacity;
+
+					FDrawDebugSkeletonSettings SkeletonSettings;
+					SkeletonSettings.bDrawSimpleSkeleton = ToolKit->Controller->ViewportSettings->bDrawSimpleSkeleton;
+					SkeletonSettings.BoneRadius = ToolKit->Controller->ViewportSettings->DrawSkeletonScale;
+
+					UDrawDebugLibrary::DrawDebugSkeletonArrayView(
+						FDebugDrawer::MakeDebugDrawer(PDI),
+						PoseState.CurrPoseData->RootData.RootLocations[CharacterIdx],
+						PoseState.CurrPoseData->RootData.RootRotations[CharacterIdx],
+						PoseState.PoseGlobalBoneData->BoneLocations[CharacterIdx],
+						PoseState.PoseGlobalBoneData->BoneRotations[CharacterIdx],
+						BoneIndices,
+						ToolKit->Controller->AutoEncoder->GetBoneParents(),
+						BoneColor,
+						false,
+						SkeletonSettings);
+
+					// Draw Velocities
+
+					FLinearColor VelocityBoneColor = BoneColor;
+					VelocityBoneColor.A = ToolKit->Controller->ViewportSettings->DrawLinearVelocitiesOpacity;
+
+					if (ToolKit->Controller->ViewportSettings->bDrawLinearVelocities)
+					{
+						UDrawDebugLibrary::DrawDebugBoneVelocitiesArrayView(
+							FDebugDrawer::MakeDebugDrawer(PDI),
+							PoseState.PoseGlobalBoneData->BoneLocations[CharacterIdx],
+							PoseState.PoseGlobalBoneData->BoneLinearVelocities[CharacterIdx],
+							BoneIndices,
+							VelocityBoneColor,
+							ToolKit->Controller->ViewportSettings->DrawSkeletonScale * 0.5f,
+							false,
+							ToolKit->Controller->ViewportSettings->DrawLinearVelocitiesScale);
+					}
+
+					if (ToolKit->Controller->ViewportSettings->bDrawAngularVelocities)
+					{
+						UDrawDebugLibrary::DrawDebugBoneVelocitiesArrayView(
+							FDebugDrawer::MakeDebugDrawer(PDI),
+							PoseState.PoseGlobalBoneData->BoneLocations[CharacterIdx],
+							PoseState.PoseGlobalBoneData->BoneAngularVelocities[CharacterIdx],
+							BoneIndices,
+							VelocityBoneColor,
+							ToolKit->Controller->ViewportSettings->DrawSkeletonScale * 0.5f,
+							false,
+							ToolKit->Controller->ViewportSettings->DrawAngularVelocitiesScale);
+					}
+
+					// Draw Transforms
+
+					if (ToolKit->Controller->ViewportSettings->bDrawBoneTransforms)
+					{
+						UDrawDebugLibrary::DrawDebugTransformsArrayView(
+							FDebugDrawer::MakeDebugDrawer(PDI),
+							PoseState.PoseGlobalBoneData->BoneLocations[CharacterIdx],
+							PoseState.PoseGlobalBoneData->BoneRotations[CharacterIdx],
+							PoseState.PoseGlobalBoneData->BoneScales[CharacterIdx],
+							UDrawDebugLibrary::MakeDrawDebugLineStyleFromColor(BoneColor),
+							false,
+							ToolKit->Controller->ViewportSettings->DrawSkeletonScale * 5.0f);
+					}
+				}
+			}
+
+			// Draw Range Identifiers
+
+			if (ToolKit->Controller->ViewportSettings->bDrawRangeIdentifier)
+			{
+				for (int32 CharacterIdx = 0; CharacterIdx < CharacterNum; CharacterIdx++)
+				{
+					FLinearColor IdentifierColor = CharacterColors[CharacterIdx];
+					IdentifierColor.A = ToolKit->Controller->ViewportSettings->DrawRangeIdentifierOpacity;
+
+					const FVector IdentifierLocation = (FVector)PoseState.CurrPoseData->RootData.RootLocations[CharacterIdx] +
+						(FVector)PoseState.CurrPoseData->RootData.RootRotations[CharacterIdx].RotateVector(FVector3f(0.0f, 0.0f, ToolKit->Controller->ViewportSettings->DrawRangeIdentifierHeight));
+
+					FDrawDebugLineStyle LineStyle;
+					LineStyle.Color = IdentifierColor;
+					LineStyle.Thickness = ToolKit->Controller->ViewportSettings->DrawRangeIdentifierThickness;
+
+					UDrawDebugLibrary::DrawDebugSimpleSphere(
+						FDebugDrawer::MakeDebugDrawer(PDI),
+						IdentifierLocation,
+						FRotator::ZeroRotator,
+						LineStyle,
+						true,
+						ToolKit->Controller->ViewportSettings->DrawRangeIdentifierRadius);
+				}
+			}
+
+			// Draw Control Encoding
+
+			if (ToolKit->Controller->ViewportSettings->bDrawEncodedControlVector)
+			{
+				const int32 EncodedControlVectorSize = ControllerOutputVector.Num<1>();
+
+				const int32 DimNum = FMath::Min(EncodedControlVectorSize, ToolKit->Controller->ViewportSettings->DrawEncodedControlVectorDimensions);
+				const float VerticalScale = ToolKit->Controller->ViewportSettings->DrawEncodedControlVectorScale;
+				const FLinearColor DrawColor = ToolKit->Controller->ViewportSettings->DrawEncodedControlVectorColor;
+
+				FDrawDebugLineStyle LineStyle;
+				LineStyle.Color = DrawColor;
+				LineStyle.Thickness = 2.0f;
+
+				if (DimNum > 0 && ControllerOutputVector.Num<0>() > 0)
+				{
+					for (int32 CharacterIdx = 0; CharacterIdx < CharacterNum; CharacterIdx++)
+					{
+						const FVector Start = PoseState.CurrPoseData->RootData.RootLocations[CharacterIdx] + FVector(50.0f, 0.0f, 100.0f);
+						const FVector Offset = FVector(20.0f, 0.0f, 0.0f);
+
+						UDrawDebugLibrary::DrawDebugLine(FDebugDrawer::MakeDebugDrawer(PDI), Start, Start + Offset * (DimNum - 1), LineStyle);
+
+						for (int32 DimIdx = 0; DimIdx < DimNum; DimIdx++)
+						{
+							const float Value = ControllerOutputVector[CharacterIdx][DimIdx];
+							const FVector DimStart = Start + Offset * DimIdx;
+							const FVector DimEnd = Start + Offset * DimIdx + FVector(0.0f, 0.0f, 10.0f * VerticalScale * Value);
+
+							FDrawDebugStringSettings StringSettings;
+							StringSettings.Height = 6.0f;
+
+							UDrawDebugLibrary::DrawDebugLine(FDebugDrawer::MakeDebugDrawer(PDI), DimStart, DimEnd, LineStyle);
+							UDrawDebugLibrary::DrawDebugSphere(FDebugDrawer::MakeDebugDrawer(PDI), DimEnd, FRotator::ZeroRotator, LineStyle, true, 5.0f);
+							UDrawDebugLibrary::DrawDebugString(FDebugDrawer::MakeDebugDrawer(PDI), FString::Printf(TEXT("% 5.2f"), Value), DimEnd + FVector(-8.0f, 0.0f, 15.0f), FRotator::ZeroRotator, UDrawDebugLibrary::DrawDebugLineStyleWithThickness(LineStyle, 0.0), true, StringSettings);
+						}
+					}
+				}
+			}
+
+			// Draw Skeleton Roots
+
+			if (ToolKit->Controller->ViewportSettings->bDrawRoot && CharacterNum > 0)
+			{
+				FDrawDebugLineStyle LineStyle;
+				LineStyle.Color = FLinearColor::White;
+				LineStyle.Color.A = ToolKit->Controller->ViewportSettings->DrawRootOpacity;
+				LineStyle.Thickness = ToolKit->Controller->ViewportSettings->DrawSkeletonScale;
+
+				UDrawDebugLibrary::DrawDebugRotationsQuatArrayView(
+					FDebugDrawer::MakeDebugDrawer(PDI),
+					PoseState.CurrPoseData->RootData.RootLocations,
+					PoseState.CurrPoseData->RootData.RootRotations,
+					LineStyle,
+					true,
+					ToolKit->Controller->ViewportSettings->DrawRootScale);
+			}
+
+			// Draw Origin
+
+			if (ToolKit->Controller->ViewportSettings->bDrawOrigin)
+			{
+				FDrawDebugLineStyle LineStyle;
+				LineStyle.Color = FLinearColor::White;
+				LineStyle.Color.A = ToolKit->Controller->ViewportSettings->DrawOriginOpacity;
+				LineStyle.Thickness = ToolKit->Controller->ViewportSettings->DrawOriginLineThickness;
+
+				UDrawDebugLibrary::DrawDebugTransform(
+					FDebugDrawer::MakeDebugDrawer(PDI),
+					FTransform::Identity,
+					LineStyle,
+					true,
+					ToolKit->Controller->ViewportSettings->DrawOriginScale);
+			}
+
+			// Debug Draw IK
+
+			if (ToolKit->Controller->TestingSettings->bDrawFootLocking)
+			{
+				for (int32 CharacterIdx = 0; CharacterIdx < CharacterNum; CharacterIdx++)
+				{
+					for (int32 BoneIdx = 0; BoneIdx < 2; BoneIdx++)
+					{
+						if (ContactState.Locked[CharacterIdx][BoneIdx])
+						{
+							FDrawDebugLineStyle LineStyle;
+							LineStyle.Color = FLinearColor::Red;
+							LineStyle.Thickness = 1.0f;
+								
+							UDrawDebugLibrary::DrawDebugSimpleSphere(FDebugDrawer::MakeDebugDrawer(PDI), ContactState.Point[CharacterIdx][BoneIdx], FRotator::ZeroRotator, LineStyle, true, 5.0f);
+						}
+
+						FDrawDebugLineStyle LineStyle;
+						LineStyle.Color = FLinearColor::Green;
+						LineStyle.Thickness = 1.0f;
+
+						UDrawDebugLibrary::DrawDebugSimpleSphere(FDebugDrawer::MakeDebugDrawer(PDI), ContactState.Position[CharacterIdx][BoneIdx], FRotator::ZeroRotator, LineStyle, true, 5.0f);
+					}
+				}
+			}
+
+			// Debug Draw
+
+			if (ToolKit->Controller->ViewportSettings->bDrawDebugTesting)
+			{
+				for (int32 CharacterIdx = 0; CharacterIdx < CharacterNum; CharacterIdx++)
+				{
+					FLinearColor IdentifierColor = CharacterColors[CharacterIdx];
+					IdentifierColor.A = ToolKit->Controller->ViewportSettings->DrawRangeIdentifierOpacity;
+
+					PoseState.PoseIdx = CharacterIdx;
+
+					BehaviorPreviews[CharacterIdx]->OnPreviewDrawDebug(FDebugDrawer::MakeDebugDrawer(PDI), ControlObject, ControlObjectElements[CharacterIdx], PoseState, IdentifierColor);
+				}
+			}
+		}
+	}
+
+	void FControllerMode::Render(const FSceneView* View, FViewport* Viewport, FPrimitiveDrawInterface* PDI)
+	{
+		if (TSharedPtr<FControllerToolkit> ToolKit = WeakToolkit.Pin())
+		{
+			if (ToolKit->ControllerState == EControllerState::Stopped)
+			{
+				RenderTraining(View, Viewport, PDI);
+			}
+			else
+			{
+				RenderTesting(View, Viewport, PDI);
+			}
+		}
+
+		FEdMode::Render(View, Viewport, PDI);
+	}
+}
+
+#undef LOCTEXT_NAMESPACE
\ No newline at end of file
diff --git a/Engine/Plugins/Experimental/Animation/AnimGen/Source/AnimGenEditor/Private/SAnimGenEditorTraining.cpp b/Engine/Plugins/Experimental/Animation/AnimGen/Source/AnimGenEditor/Private/SAnimGenEditorTraining.cpp
new file mode 100644
--- /dev/null
+++ b/Engine/Plugins/Experimental/Animation/AnimGen/Source/AnimGenEditor/Private/SAnimGenEditorTraining.cpp
@@ -0,0 +1,157 @@
+// Copyright Epic Games, Inc. All Rights Reserved.
+
+#include "SAnimGenEditorTraining.h"
+#include "AnimGenEditorTraining.h"
+
+#include "Widgets/SBoxPanel.h"
+#include "Widgets/Input/SButton.h"
+#include "Widgets/Notifications/SProgressBar.h"
+#include "Widgets/Text/STextBlock.h"
+
+#define LOCTEXT_NAMESPACE "AnimGenEditorTraining"
+
+namespace UE::AnimGen::Editor
+{
+	void STraining::Construct(const FArguments& InArgs, const TWeakPtr<ITrainingModel> InTrainingModel)
+	{
+		TrainingModel = InTrainingModel;
+
+		ChildSlot
+		[
+			SNew(SVerticalBox)
+				+SVerticalBox::Slot()
+				.Padding(10.0f, 10.0f)
+				.MaxHeight(30.0f)
+				[
+					SNew(SHorizontalBox)
+						+ SHorizontalBox::Slot()
+						.HAlign(HAlign_Left)
+						.VAlign(VAlign_Center)
+						.Padding(10.0f, 0.0f)
+						.AutoWidth()
+						[
+							SNew(SButton)
+								.Text(LOCTEXT("Retrain", "Retrain"))
+								.IsEnabled_Lambda([this]() { return TrainingModel.IsValid() ? !TrainingModel.Pin()->IsTraining() : false;  })
+								.OnClicked_Lambda([this]() { if (TrainingModel.IsValid()) { TrainingModel.Pin()->Retrain(); return FReply::Handled(); } return FReply::Unhandled(); })
+						]
+						+ SHorizontalBox::Slot()
+						.HAlign(HAlign_Left)
+						.VAlign(VAlign_Center)
+						.Padding(10.0f, 0.0f)
+						.AutoWidth()
+						[
+							SNew(SButton)
+								.Text(LOCTEXT("FineTune", "Fine Tune"))
+								.Visibility_Lambda([this]() { return TrainingModel.Pin()->SupportsFinetune() ? EVisibility::Visible : EVisibility::Collapsed; })
+								.IsEnabled_Lambda([this]() { return TrainingModel.IsValid() ? !TrainingModel.Pin()->IsTraining() && !TrainingModel.Pin()->IsEmpty() : false;  })
+								.OnClicked_Lambda([this]() { if (TrainingModel.IsValid()) { TrainingModel.Pin()->FineTune(); return FReply::Handled(); } return FReply::Unhandled(); })
+						]
+						+ SHorizontalBox::Slot()
+						.HAlign(HAlign_Left)
+						.VAlign(VAlign_Center)
+						.Padding(10.0f, 0.0f)
+						.AutoWidth()
+						[
+							SNew(SButton)
+								.Text(LOCTEXT("Distill", "Distill"))
+								.Visibility_Lambda([this]() { return TrainingModel.Pin()->SupportsDistill() ? EVisibility::Visible : EVisibility::Collapsed; })
+								.IsEnabled_Lambda([this]() { return TrainingModel.IsValid() ? !TrainingModel.Pin()->IsTraining() && TrainingModel.Pin()->SupportsDistill() : false;  })
+								.OnClicked_Lambda([this]() { if (TrainingModel.IsValid()) { TrainingModel.Pin()->Distill(); return FReply::Handled(); } return FReply::Unhandled(); })
+						]
+						+ SHorizontalBox::Slot()
+						.HAlign(HAlign_Left)
+						.VAlign(VAlign_Center)
+						.Padding(10.0f, 0.0f)
+						.AutoWidth()
+						[
+							SNew(SButton)
+								.Text(LOCTEXT("Stop", "Stop"))
+								.IsEnabled_Lambda([this]() { return TrainingModel.IsValid() ? TrainingModel.Pin()->IsTraining() : false;  })
+								.OnClicked_Lambda([this]() { if (TrainingModel.IsValid()) { TrainingModel.Pin()->Stop(); return FReply::Handled(); } return FReply::Unhandled(); })
+						]
+
+				]
+				+SVerticalBox::Slot()
+				.Padding(10.0f, 10.0f)
+				.MaxHeight(20.0f)
+				[
+					SNew(SProgressBar)
+						.Percent_Lambda([this]()
+							{ 
+								if (TrainingModel.IsValid())
+								{
+									const int32 Iteration = TrainingModel.Pin()->GetIterationNum();
+									const int32 MaxIteration = TrainingModel.Pin()->GetMaxIterationNum();
+									return MaxIteration > 1 ? TOptional<float>((float)Iteration / (MaxIteration - 1)) : TOptional<float>(0.0f);
+								}
+
+								return TOptional<float>(0.0f);
+							})
+				]
+				+SVerticalBox::Slot()
+				.Padding(10.0f, 10.0f)
+				.MaxHeight(30.0f)
+				[
+					SNew(STextBlock).Text_Lambda([this]()
+						{
+							if (TrainingModel.IsValid())
+							{
+								const TSharedPtr<ITrainingModel> Model = TrainingModel.Pin();
+
+								FText ErrorMessage = Model->GetErrorMessage();
+
+								if (!ErrorMessage.IsEmpty()) { return ErrorMessage; }
+
+								switch (Model->GetTrainingStatus())
+								{
+								case ETrainingStatus::NotStarted: return LOCTEXT("NotStartedStatus", "Status: Training not started.");
+								case ETrainingStatus::Preparing: return LOCTEXT("PreparingStatus", "Status: Preparing training...");
+								case ETrainingStatus::Training:
+								{
+									TOptional<FTimespan> RemainingTime = Model->GetEstimateTimeRemaining();
+
+									FText TimeRemainingText;
+									if (RemainingTime.IsSet())
+									{
+										FNumberFormattingOptions DigitsFormatOptions;
+										DigitsFormatOptions.SetMinimumIntegralDigits(2);
+										DigitsFormatOptions.SetMaximumIntegralDigits(2);
+
+										TimeRemainingText = FText::Format(
+											LOCTEXT("RemainingTime", ", time remaining {0}:{1}:{2}"),
+											FText::AsNumber(RemainingTime->GetDays() * 24 + RemainingTime->GetHours(), &DigitsFormatOptions),
+											FText::AsNumber(RemainingTime->GetMinutes(), &DigitsFormatOptions),
+											FText::AsNumber(RemainingTime->GetSeconds(), &DigitsFormatOptions));
+									}
+
+									FNumberFormattingOptions LossFormatOptions;
+									LossFormatOptions.SetMinimumFractionalDigits(5);
+									LossFormatOptions.SetMaximumFractionalDigits(5);
+
+									FNumberFormattingOptions IterationFormatOptions;
+									IterationFormatOptions.SetMinimumIntegralDigits(7);
+									IterationFormatOptions.SetUseGrouping(false);
+
+									return FText::Format(LOCTEXT("TrainingStatus", "Status: Training iteration {0} of {1}, loss {2}{3}"),
+										FText::AsNumber(Model->GetIterationNum(), &IterationFormatOptions),
+										FText::AsNumber(Model->GetMaxIterationNum(), &IterationFormatOptions),
+										FText::AsNumber(Model->GetTrainingLoss(), &LossFormatOptions),
+										TimeRemainingText);
+								}
+								case ETrainingStatus::Done: return LOCTEXT("DoneStatus", "Status: Training complete.");
+								default: return LOCTEXT("ErrorStatus", "Error");
+								}
+							}
+							else
+							{
+								return FText();
+							}
+						})
+				]
+		];
+
+	}
+}
+
+#undef LOCTEXT_NAMESPACE
\ No newline at end of file
