AnimGen: Move to Experimental

diff --git a/Engine/Plugins/Experimental/Animation/AnimGen/Source/AnimGenEditor/Public/AnimGenEditorControllerMode.h b/Engine/Plugins/Experimental/Animation/AnimGen/Source/AnimGenEditor/Public/AnimGenEditorControllerMode.h
new file mode 100644
--- /dev/null
+++ b/Engine/Plugins/Experimental/Animation/AnimGen/Source/AnimGenEditor/Public/AnimGenEditorControllerMode.h
@@ -0,0 +1,143 @@
+// Copyright Epic Games, Inc. All Rights Reserved.
+
+#pragma once
+
+#include "LearningArray.h"
+
+#include "AnimGenControl.h"
+#include "AnimGenEditorPerfCounter.h"
+
+#include "AnimDatabase.h"
+#include "AnimDatabasePose.h"
+
+#include "EdMode.h"
+
+namespace UE::Learning
+{
+	struct FNeuralNetworkInference;
+}
+
+class ULearningNeuralNetworkData;
+class UAnimGenBehavior;
+class UAnimGenBehaviorPreview;
+struct FAnimDatabaseFrameAttributeEntry;
+
+namespace UE::AnimGen::Editor
+{
+	class FControllerToolkit;
+
+	class FControllerMode : public FEdMode
+	{
+	public:
+
+		FControllerMode();
+		~FControllerMode();
+
+		const static FEditorModeID EditorModeId;
+
+		virtual void Tick(FEditorViewportClient* ViewportClient, float DeltaTime) override;
+		virtual void Render(const FSceneView* View, FViewport* Viewport, FPrimitiveDrawInterface* PDI);
+
+		virtual bool MouseMove(FEditorViewportClient* ViewportClient, FViewport* Viewport, int32 X, int32 Y) override;
+		virtual bool CapturedMouseMove(FEditorViewportClient* ViewportClient, FViewport* Viewport, int32 X, int32 Y) override;
+		virtual bool InputAxis(FEditorViewportClient* InViewportClient, FViewport* Viewport, int32 ControllerId, FKey Key, float Delta, float DeltaTime) override;
+		virtual bool InputKey(FEditorViewportClient* ViewportClient, FViewport* Viewport, FKey Key, EInputEvent Event) override;
+
+		void TickTraining(FEditorViewportClient* EditorViewportClient, float DeltaTime);
+		void TickTesting(FEditorViewportClient* EditorViewportClient, float DeltaTime);
+
+		void RenderTraining(const FSceneView* View, FViewport* Viewport, FPrimitiveDrawInterface* PDI);
+		void RenderTesting(const FSceneView* View, FViewport* Viewport, FPrimitiveDrawInterface* PDI);
+
+	private:
+
+		/** Maximum number of characters to draw in the scene */
+		const int32 MaxCharacterNum = 128;
+
+		bool bCaptureInput = false;
+
+		int32 CharacterNum = 0;
+		TArray<FLinearColor> CharacterColors;
+
+		/** Weak pointer to the controller toolkit used for rendering */
+		TWeakPtr<FControllerToolkit> WeakToolkit;
+
+		TArray<int32> ActiveRanges;
+
+		/** Character Data */
+		TArray<int32> CharacterSequenceIndices;
+		TArray<float> CharacterRangeTimes;
+		TArray<int32> CharacterRangeStarts;
+		TArray<int32> CharacterRangeLengths;
+		TArray<int32> CharacterControlSetIndices;
+		TArray<int32> CharacterControlSetEntries;
+		TArray<int32> CharacterControlSetRanges;
+		TArray<int32> CharacterControlSetFrames;
+
+		/** Trajectories for all currently queried ranges */
+		TArray<TLearningArray<1, FVector>> RangeTrajectoryLocations;
+		TArray<TLearningArray<1, FQuat4f>> RangeTrajectoryRotations;
+
+		/** Current Frame Attribute Entries */
+		TArray<FAnimDatabaseFrameAttributeEntry> FrameAttributeEntries;
+
+		/** Associated Frame Attributes */
+		TArray<FAnimDatabaseFrameAttribute> FrameAttributes;
+		
+		/** Frame Attributes database hash */
+		int32 FrameAttributesDatabaseContentHash = 0;
+
+		/** By default we should reset at the start */
+		bool bReset = true;
+
+		/** Trained Database content hash */
+		int32 TrainedDatabaseContentHash = 0;
+		UObject* TrainedDatabaseFrameRangesDefaultObject = nullptr;
+		int32 TrainedAutoEncoderContentHash = 0;
+
+		TObjectPtr<UAnimGenBehaviorPreview> ActiveBehaviorObject = nullptr;
+		TArray<UAnimGenBehaviorPreview*> BehaviorPreviews;
+
+		TSharedPtr<UE::Learning::FNeuralNetworkInference> EncoderInference;
+		TSharedPtr<UE::Learning::FNeuralNetworkInference> DecoderInference;
+		TSharedPtr<UE::Learning::FNeuralNetworkInference> ControlEncoderInference;
+		TSharedPtr<UE::Learning::FNeuralNetworkInference> LOD0Inference;
+		TSharedPtr<UE::Learning::FNeuralNetworkInference> LOD1Inference;
+		TSharedPtr<UE::Learning::FNeuralNetworkInference> LOD2Inference;
+
+		TObjectPtr<ULearningNeuralNetworkData> EncoderInferenceNetwork;
+		TObjectPtr<ULearningNeuralNetworkData> DecoderInferenceNetwork;
+		TObjectPtr<ULearningNeuralNetworkData> ControlEncoderInferenceNetwork;
+		TObjectPtr<ULearningNeuralNetworkData> LOD0InferenceNetwork;
+		TObjectPtr<ULearningNeuralNetworkData> LOD1InferenceNetwork;
+		TObjectPtr<ULearningNeuralNetworkData> LOD2InferenceNetwork;
+
+		UE::AnimGen::Editor::FPerfCounter ControlEncoderPerfCounter;
+		UE::AnimGen::Editor::FPerfCounter DenoiserPerfCounter;
+
+		TArray<TPair<FKey, EInputEvent>, TInlineAllocator<8>> KeyEvents;
+
+		FAnimDatabasePoseState PoseState;
+
+		TLearningArray<2, float> MemoryVector;
+		TLearningArray<2, float> ControlVector;
+		TLearningArray<2, float> PoseVector;
+		TLearningArray<2, float> UnnormalizedEncodedPoseVector;
+		TLearningArray<2, float> EncodedPoseVector;
+		TLearningArray<2, float> NextEncodedPoseVector;
+		TLearningArray<2, float> ControllerInputVector;
+		TLearningArray<2, float> ControllerOutputVector;
+		TLearningArray<2, float> DenoiserInputVector;
+
+		UE::AnimDatabase::FPoseData OutputPoseData;
+		TLearningArray<2, float> PredictedPoseVector;
+		UE::AnimDatabase::FPoseData AdjustedPoseData;
+		UE::AnimDatabase::FPoseGlobalBoneData AdjustedGlobalBoneData;
+		UE::AnimDatabase::PoseData::ContactPinning::FContactState ContactState;
+
+		uint32 RandomState = 1234;
+
+		FAnimGenControlObject ControlObject;
+		TArray<FAnimGenControlObjectElement> ControlObjectElements;
+	};
+}
diff --git a/Engine/Plugins/Experimental/Animation/AnimGen/Source/AnimGenEditor/Public/AnimGenEditorTraining.h b/Engine/Plugins/Experimental/Animation/AnimGen/Source/AnimGenEditor/Public/AnimGenEditorTraining.h
new file mode 100644
--- /dev/null
+++ b/Engine/Plugins/Experimental/Animation/AnimGen/Source/AnimGenEditor/Public/AnimGenEditorTraining.h
@@ -0,0 +1,36 @@
+// Copyright Epic Games, Inc. All Rights Reserved.
+
+#pragma once
+
+#include "CoreTypes.h"
+#include "Internationalization/Text.h"
+#include "Misc/DateTime.h"
+
+namespace UE::AnimGen::Editor
+{
+	enum class ETrainingStatus : uint8
+	{
+		NotStarted = 0,
+		Preparing = 1,
+		Training = 2,
+		Done = 3,
+	};
+
+	struct ITrainingModel
+	{
+		virtual void Retrain() {}
+		virtual void FineTune() {}
+		virtual void Distill() {}
+		virtual void Stop() {}
+		virtual bool IsEmpty() const { return true; }
+		virtual bool IsTraining() const { return false; }
+		virtual bool SupportsDistill() const { return false; }
+		virtual bool SupportsFinetune() const { return false; }
+		virtual int32 GetIterationNum() const { return 0; }
+		virtual int32 GetMaxIterationNum() const { return 0; }
+		virtual ETrainingStatus GetTrainingStatus() const { return ETrainingStatus::NotStarted; }
+		virtual float GetTrainingLoss() const { return 0.0f; }
+		virtual FText GetErrorMessage() const { return FText(); }
+		virtual TOptional<FTimespan> GetEstimateTimeRemaining() const { return TOptional<FTimespan>(); };
+	};
+}
\ No newline at end of file
