AnimGen: Added asset creation from right click context menu

#Synced-CL 51417185

diff --git a/Engine/Plugins/Experimental/Animation/AnimDatabase/Source/AnimDatabaseEditor/Private/AnimDatabaseEditorToolkit.cpp b/Engine/Plugins/Experimental/Animation/AnimDatabase/Source/AnimDatabaseEditor/Private/AnimDatabaseEditorToolkit.cpp
--- a/Engine/Plugins/Experimental/Animation/AnimDatabase/Source/AnimDatabaseEditor/Private/AnimDatabaseEditorToolkit.cpp
+++ b/Engine/Plugins/Experimental/Animation/AnimDatabase/Source/AnimDatabaseEditor/Private/AnimDatabaseEditorToolkit.cpp
@@ -71,6 +71,48 @@ EAssetCommandResult UAnimDatabaseEditorAssetDefinition::OpenAssets(const FAssetO
 	return EAssetCommandResult::Handled;
 }
 
+UAnimDatabaseEditorFactory::UAnimDatabaseEditorFactory()
+{
+	bCreateNew = true;
+	bEditAfterNew = true;
+	SupportedClass = UAnimDatabase::StaticClass();
+}
+
+UObject* UAnimDatabaseEditorFactory::FactoryCreateNew(UClass* Class, UObject* InParent, FName Name, EObjectFlags Flags, UObject* Context, FFeedbackContext* Warn)
+{
+	return NewObject<UAnimDatabase>(InParent, Name, Flags | RF_Transactional);
+}
+
+bool UAnimDatabaseEditorFactory::ShouldShowInNewMenu() const
+{
+	return true;
+}
+
+bool UAnimDatabaseEditorFactory::ConfigureProperties()
+{
+	return true;
+}
+
+FText UAnimDatabaseEditorFactory::GetDisplayName() const
+{
+	return LOCTEXT("AnimDatabaseEditorAsset_DisplayName", "Animation Database");
+}
+
+uint32 UAnimDatabaseEditorFactory::GetMenuCategories() const
+{
+	return EAssetTypeCategories::Animation;
+}
+
+FText UAnimDatabaseEditorFactory::GetToolTip() const
+{
+	return LOCTEXT("AnimDatabaseEditorAsset_Tooltip", "Asset for organization and processing large databases of animation.");
+}
+
+FString UAnimDatabaseEditorFactory::GetDefaultNewAssetName() const
+{
+	return FString(TEXT("ADB_NewAnimDatabase"));
+}
+
 namespace UE::AnimDatabase::Editor
 {
 	TSharedRef<IDetailCustomization> FQueryDetails::MakeInstance()
diff --git a/Engine/Plugins/Experimental/Animation/AnimDatabase/Source/AnimDatabaseEditor/Public/AnimDatabaseEditorToolkit.h b/Engine/Plugins/Experimental/Animation/AnimDatabase/Source/AnimDatabaseEditor/Public/AnimDatabaseEditorToolkit.h
--- a/Engine/Plugins/Experimental/Animation/AnimDatabase/Source/AnimDatabaseEditor/Public/AnimDatabaseEditorToolkit.h
+++ b/Engine/Plugins/Experimental/Animation/AnimDatabase/Source/AnimDatabaseEditor/Public/AnimDatabaseEditorToolkit.h
@@ -6,6 +6,7 @@
 #include "AssetDefinitionDefault.h"
 #include "IDetailCustomization.h"
 #include "Widgets/Views/SListView.h"
+#include "Factories/Factory.h"
 
 #include "AnimDatabaseEditorToolkit.generated.h"
 
@@ -31,6 +32,26 @@ public:
 	virtual EAssetCommandResult OpenAssets(const FAssetOpenArgs& OpenArgs) const override final;
 };
 
+/** Asset factory for the Animation Database class */
+UCLASS(MinimalAPI, hidecategories = Object)
+class UAnimDatabaseEditorFactory : public UFactory
+{
+	GENERATED_BODY()
+
+public:
+	UAnimDatabaseEditorFactory();
+
+	// UFactory overrides.
+	virtual UObject* FactoryCreateNew(UClass* Class, UObject* InParent, FName Name, EObjectFlags Flags, UObject* Context, FFeedbackContext* Warn) override;
+	virtual FText GetDisplayName() const override;
+	virtual uint32 GetMenuCategories() const override;
+	virtual FText GetToolTip() const override;
+	virtual FString GetDefaultNewAssetName() const override;
+	virtual bool ConfigureProperties() override;
+	virtual bool ShouldShowInNewMenu() const override;
+	// ~END UFactory overrides.
+};
+
 namespace UE::AnimDatabase::Editor
 {
 	class FPreviewScene;
diff --git a/Engine/Plugins/Experimental/Animation/AnimGen/Source/AnimGenEditor/Private/AnimGenEditorAutoEncoderToolkit.cpp b/Engine/Plugins/Experimental/Animation/AnimGen/Source/AnimGenEditor/Private/AnimGenEditorAutoEncoderToolkit.cpp
--- a/Engine/Plugins/Experimental/Animation/AnimGen/Source/AnimGenEditor/Private/AnimGenEditorAutoEncoderToolkit.cpp
+++ b/Engine/Plugins/Experimental/Animation/AnimGen/Source/AnimGenEditor/Private/AnimGenEditorAutoEncoderToolkit.cpp
@@ -50,6 +50,7 @@
 #include "AdvancedPreviewSceneModule.h"
 #include "Editor.h"
 #include "Subsystems/AssetEditorSubsystem.h"
+#include "AssetTypeCategories.h"
 
 #include "Animation/Skeleton.h"
 #include "Animation/AnimSequence.h"
@@ -71,7 +72,7 @@
 
 FText UAnimGenEditorAutoEncoderAssetDefinition::GetAssetDisplayName() const
 {
-	return LOCTEXT("AnimGenEditorAutoEncoderAssetDefinitionName", "Animation Auto-Encoder");
+	return LOCTEXT("AnimGenEditorAutoEncoderAssetDefinitionName", "Auto-Encoder");
 }
 
 FLinearColor UAnimGenEditorAutoEncoderAssetDefinition::GetAssetColor() const
@@ -101,6 +102,55 @@ EAssetCommandResult UAnimGenEditorAutoEncoderAssetDefinition::OpenAssets(const F
 	return EAssetCommandResult::Handled;
 }
 
+
+UAnimGenAutoEncoderEditorFactory::UAnimGenAutoEncoderEditorFactory()
+{
+	bCreateNew = true;
+	bEditAfterNew = true;
+	SupportedClass = UAnimGenAutoEncoder::StaticClass();
+}
+
+UObject* UAnimGenAutoEncoderEditorFactory::FactoryCreateNew(UClass* Class, UObject* InParent, FName Name, EObjectFlags Flags, UObject* Context, FFeedbackContext* Warn)
+{
+	return NewObject<UAnimGenAutoEncoder>(InParent, Name, Flags | RF_Transactional);
+}
+
+bool UAnimGenAutoEncoderEditorFactory::ShouldShowInNewMenu() const
+{
+	return true;
+}
+
+bool UAnimGenAutoEncoderEditorFactory::ConfigureProperties()
+{
+	return true;
+}
+
+FText UAnimGenAutoEncoderEditorFactory::GetDisplayName() const
+{
+	return LOCTEXT("AnimGenAutoEncoderEditorAsset_DisplayName", "Auto-Encoder");
+}
+
+uint32 UAnimGenAutoEncoderEditorFactory::GetMenuCategories() const
+{
+	return EAssetTypeCategories::Animation;
+}
+
+FText UAnimGenAutoEncoderEditorFactory::GetToolTip() const
+{
+	return LOCTEXT("AnimGenAutoEncoderEditorAsset_Tooltip", "Asset for training a model which creates a compressed encoding of animation data suitable for machine learning tasks.");
+}
+
+FString UAnimGenAutoEncoderEditorFactory::GetDefaultNewAssetName() const
+{
+	return FString(TEXT("AGAE_NewAnimGenAutoEncoder"));
+}
+
+const TArray<FText>& UAnimGenAutoEncoderEditorFactory::GetMenuCategorySubMenus() const
+{
+	static TArray<FText> SubMenus{ LOCTEXT("SubMenuAnimGen", "AnimGen") };
+	return SubMenus;
+}
+
 namespace UE::AnimGen::Editor
 {
 	TSharedRef<IDetailCustomization> FAutoEncoderTrainingSettingsDetails::MakeInstance()
diff --git a/Engine/Plugins/Experimental/Animation/AnimGen/Source/AnimGenEditor/Private/AnimGenEditorControllerToolkit.cpp b/Engine/Plugins/Experimental/Animation/AnimGen/Source/AnimGenEditor/Private/AnimGenEditorControllerToolkit.cpp
--- a/Engine/Plugins/Experimental/Animation/AnimGen/Source/AnimGenEditor/Private/AnimGenEditorControllerToolkit.cpp
+++ b/Engine/Plugins/Experimental/Animation/AnimGen/Source/AnimGenEditor/Private/AnimGenEditorControllerToolkit.cpp
@@ -47,6 +47,7 @@
 #include "Widgets/Input/SCheckBox.h"
 #include "Widgets/Input/SHyperlink.h"
 #include "Widgets/Layout/SScrollBox.h"
+#include "AssetTypeCategories.h"
 
 #include "Animation/Skeleton.h"
 #include "Animation/AnimSequence.h"
@@ -69,7 +70,7 @@
 
 FText UAnimGenEditorControllerAssetDefinition::GetAssetDisplayName() const
 {
-	return LOCTEXT("AnimGenEditorControllerAssetDefinitionName", "Animation Controller");
+	return LOCTEXT("AnimGenEditorControllerAssetDefinitionName", "Controller");
 }
 
 FLinearColor UAnimGenEditorControllerAssetDefinition::GetAssetColor() const
@@ -99,6 +100,54 @@ EAssetCommandResult UAnimGenEditorControllerAssetDefinition::OpenAssets(const FA
 	return EAssetCommandResult::Handled;
 }
 
+UAnimGenControllerEditorFactory::UAnimGenControllerEditorFactory()
+{
+	bCreateNew = true;
+	bEditAfterNew = true;
+	SupportedClass = UAnimGenController::StaticClass();
+}
+
+UObject* UAnimGenControllerEditorFactory::FactoryCreateNew(UClass* Class, UObject* InParent, FName Name, EObjectFlags Flags, UObject* Context, FFeedbackContext* Warn)
+{
+	return NewObject<UAnimGenController>(InParent, Name, Flags | RF_Transactional);
+}
+
+bool UAnimGenControllerEditorFactory::ShouldShowInNewMenu() const
+{
+	return true;
+}
+
+bool UAnimGenControllerEditorFactory::ConfigureProperties()
+{
+	return true;
+}
+
+FText UAnimGenControllerEditorFactory::GetDisplayName() const
+{
+	return LOCTEXT("AnimGenControllerEditorAsset_DisplayName", "Controller");
+}
+
+uint32 UAnimGenControllerEditorFactory::GetMenuCategories() const
+{
+	return EAssetTypeCategories::Animation;
+}
+
+FText UAnimGenControllerEditorFactory::GetToolTip() const
+{
+	return LOCTEXT("AnimGenAControllerEditorAsset_Tooltip", "Asset for training a model which can create animation following a provided behavior.");
+}
+
+FString UAnimGenControllerEditorFactory::GetDefaultNewAssetName() const
+{
+	return FString(TEXT("AGC_NewAnimGenController"));
+}
+
+const TArray<FText>& UAnimGenControllerEditorFactory::GetMenuCategorySubMenus() const
+{
+	static TArray<FText> SubMenus{ LOCTEXT("SubMenuAnimGen", "AnimGen") };
+	return SubMenus;
+}
+
 namespace UE::AnimGen::Editor
 {
 	TSharedRef<IDetailCustomization> FControllerTrainingSettingsDetails::MakeInstance()
diff --git a/Engine/Plugins/Experimental/Animation/AnimGen/Source/AnimGenEditor/Public/AnimGenEditorAutoEncoderToolkit.h b/Engine/Plugins/Experimental/Animation/AnimGen/Source/AnimGenEditor/Public/AnimGenEditorAutoEncoderToolkit.h
--- a/Engine/Plugins/Experimental/Animation/AnimGen/Source/AnimGenEditor/Public/AnimGenEditorAutoEncoderToolkit.h
+++ b/Engine/Plugins/Experimental/Animation/AnimGen/Source/AnimGenEditor/Public/AnimGenEditorAutoEncoderToolkit.h
@@ -9,6 +9,7 @@
 #include "AssetDefinitionDefault.h"
 #include "IDetailCustomization.h"
 #include "Widgets/Views/SListView.h"
+#include "Factories/Factory.h"
 
 #include "AnimGenEditorAutoEncoderToolkit.generated.h"
 
@@ -29,6 +30,27 @@ public:
 	virtual EAssetCommandResult OpenAssets(const FAssetOpenArgs& OpenArgs) const override final;
 };
 
+/** Asset factory for the AnimGen AutoEncoder class */
+UCLASS(MinimalAPI, hidecategories = Object)
+class UAnimGenAutoEncoderEditorFactory : public UFactory
+{
+	GENERATED_BODY()
+
+public:
+	UAnimGenAutoEncoderEditorFactory();
+
+	// UFactory overrides.
+	virtual UObject* FactoryCreateNew(UClass* Class, UObject* InParent, FName Name, EObjectFlags Flags, UObject* Context, FFeedbackContext* Warn) override;
+	virtual FText GetDisplayName() const override;
+	virtual uint32 GetMenuCategories() const override;
+	virtual FText GetToolTip() const override;
+	virtual FString GetDefaultNewAssetName() const override;
+	virtual bool ConfigureProperties() override;
+	virtual bool ShouldShowInNewMenu() const override;
+	virtual const TArray<FText>& GetMenuCategorySubMenus() const override;
+	// ~END UFactory overrides.
+};
+
 namespace UE::AnimDatabase::Editor
 {
 	class FPreviewScene;
diff --git a/Engine/Plugins/Experimental/Animation/AnimGen/Source/AnimGenEditor/Public/AnimGenEditorControllerToolkit.h b/Engine/Plugins/Experimental/Animation/AnimGen/Source/AnimGenEditor/Public/AnimGenEditorControllerToolkit.h
--- a/Engine/Plugins/Experimental/Animation/AnimGen/Source/AnimGenEditor/Public/AnimGenEditorControllerToolkit.h
+++ b/Engine/Plugins/Experimental/Animation/AnimGen/Source/AnimGenEditor/Public/AnimGenEditorControllerToolkit.h
@@ -9,6 +9,7 @@
 #include "AssetDefinitionDefault.h"
 #include "IDetailCustomization.h"
 #include "Widgets/Views/SListView.h"
+#include "Factories/Factory.h"
 
 #include "AnimGenEditorControllerToolkit.generated.h"
 
@@ -29,6 +30,27 @@ public:
 	virtual EAssetCommandResult OpenAssets(const FAssetOpenArgs& OpenArgs) const override final;
 };
 
+/** Asset factory for the AnimGen Controller class */
+UCLASS(MinimalAPI, hidecategories = Object)
+class UAnimGenControllerEditorFactory : public UFactory
+{
+	GENERATED_BODY()
+
+public:
+	UAnimGenControllerEditorFactory();
+
+	// UFactory overrides.
+	virtual UObject* FactoryCreateNew(UClass* Class, UObject* InParent, FName Name, EObjectFlags Flags, UObject* Context, FFeedbackContext* Warn) override;
+	virtual FText GetDisplayName() const override;
+	virtual uint32 GetMenuCategories() const override;
+	virtual FText GetToolTip() const override;
+	virtual FString GetDefaultNewAssetName() const override;
+	virtual bool ConfigureProperties() override;
+	virtual bool ShouldShowInNewMenu() const override;
+	virtual const TArray<FText>& GetMenuCategorySubMenus() const override;
+	// ~END UFactory overrides.
+};
+
 namespace UE::AnimDatabase::Editor
 {
 	class FPreviewScene;
