AnimGen: Fix for bone indexing error.

#Synced-CL 51427675

diff --git a/Engine/Plugins/Experimental/Animation/AnimGen/Source/AnimGen/Private/AnimNode_AnimGenController.cpp b/Engine/Plugins/Experimental/Animation/AnimGen/Source/AnimGen/Private/AnimNode_AnimGenController.cpp
--- a/Engine/Plugins/Experimental/Animation/AnimGen/Source/AnimGen/Private/AnimNode_AnimGenController.cpp
+++ b/Engine/Plugins/Experimental/Animation/AnimGen/Source/AnimGen/Private/AnimNode_AnimGenController.cpp
@@ -653,14 +653,19 @@ void FAnimNode_AnimGenController::Evaluate_AnyThread(FPoseContext& Output)
 		const int32 PoseBoneIdx = BoneContainer.GetSkeletonPoseIndexFromCompactPoseIndex(FCompactPoseBoneIndex(BoneIdx)).GetInt();
 		if (PoseBoneIdx != INDEX_NONE)
 		{
-			const int32 AutoEncoderBoneIdx = ActiveController->AutoEncoder->FindBoneIndex(BoneContainer.GetReferenceSkeleton().GetBoneName(PoseBoneIdx));
-			
-			if (AutoEncoderBoneIdx != INDEX_NONE)
+			const int32 MeshBoneIdx = BoneContainer.GetMeshPoseIndexFromSkeletonPoseIndex(FSkeletonPoseBoneIndex(PoseBoneIdx)).GetInt();
+
+			if (BoneContainer.GetReferenceSkeleton().IsValidIndex(MeshBoneIdx))
 			{
-				CompactBones[BoneIdx] = FTransform(
-					((FQuat)PoseState.CurrPoseData->LocalBoneData.BoneRotations[0][AutoEncoderBoneIdx]).GetNormalized(),
-					(FVector)PoseState.CurrPoseData->LocalBoneData.BoneLocations[0][AutoEncoderBoneIdx],
-					(FVector)PoseState.CurrPoseData->LocalBoneData.BoneScales[0][AutoEncoderBoneIdx]);
+				const int32 AutoEncoderBoneIdx = ActiveController->AutoEncoder->FindBoneIndex(BoneContainer.GetReferenceSkeleton().GetBoneName(MeshBoneIdx));
+
+				if (AutoEncoderBoneIdx != INDEX_NONE)
+				{
+					CompactBones[BoneIdx] = FTransform(
+						((FQuat)PoseState.CurrPoseData->LocalBoneData.BoneRotations[0][AutoEncoderBoneIdx]).GetNormalized(),
+						(FVector)PoseState.CurrPoseData->LocalBoneData.BoneLocations[0][AutoEncoderBoneIdx],
+						(FVector)PoseState.CurrPoseData->LocalBoneData.BoneScales[0][AutoEncoderBoneIdx]);
+				}
 			}
 		}
 	}
