AnimGen: Updated controller training script

#virtualized
#Synced-CL 51457572

diff --git a/Engine/Plugins/Experimental/Animation/AnimGen/Source/AnimGen/Public/AnimGenController.h b/Engine/Plugins/Experimental/Animation/AnimGen/Source/AnimGen/Public/AnimGenController.h
--- a/Engine/Plugins/Experimental/Animation/AnimGen/Source/AnimGen/Public/AnimGenController.h
+++ b/Engine/Plugins/Experimental/Animation/AnimGen/Source/AnimGen/Public/AnimGenController.h
@@ -357,10 +357,14 @@ public:
 
 public:
 
-	/** The number of iterations to run before ending training. */
+	/** The number of iterations to train the reference flow-matching network. */
 	UPROPERTY(EditAnywhere, Category = "Training", meta = (ClampMin = "1", UIMin = "1"))
 	int32 NumberOfIterations = 1000000;
 
+	/** The number of iterations to train the distilled LOD networks. */
+	UPROPERTY(EditAnywhere, Category = "Training", meta = (ClampMin = "1", UIMin = "1"))
+	int32 NumberOfDistillationIterations = 100000;
+
 	/** The number of learning rate warm-up iterations to apply. Stabilizes the beginning of training. */
 	UPROPERTY(EditAnywhere, Category = "Training", meta = (ClampMin = "1", UIMin = "1"))
 	int32 WarmupIterations = 1000;
@@ -380,23 +384,19 @@ public:
 	UPROPERTY(EditAnywhere, Category = "Training", meta = (ClampMin = "0", UIMin = "0"))
 	int32 RandomSeed = 1234;
 
-	/** Initial scale of noise to apply to the pose encoding. Larger values may make generation more stable and add variety at the expense of making the result more noisy. */
+	/** Initial scale of noise to apply to the pose encoding. */
 	UPROPERTY(EditAnywhere, Category = "Training", meta = (ClampMin = "0.0", UIMin = "0.0", UIMax = "1.0"))
 	float MaxPoseNoiseScale = 1.0f;
 
-	/** Final scale of noise to apply to the pose encoding. Larger values may make generation more stable and add variety at the expense of making the result more noisy. */
+	/** Final scale of noise to apply to the pose encoding. Larger values may allow the character to transition better between states at the expense of making the result more noisy. */
 	UPROPERTY(EditAnywhere, Category = "Training", meta = (ClampMin = "0.0", UIMin = "0.0", UIMax = "1.0"))
-	float MinPoseNoiseScale = 0.0f;
+	float MinPoseNoiseScale = 0.01f;
 
 	/** Scale of noise to apply to the control inputs. Larger values may improve generalization a little. */
 	UPROPERTY(EditAnywhere, Category = "Training", meta = (ClampMin = "0.0", UIMin = "0.0", UIMax = "1.0"))
 	float ControlNoiseScale = 0.01f;
 
-	/** Student forcing rate. Larger values should increase stability but may make the output more noisy and reduce animation quality. */
-	UPROPERTY(EditAnywhere, Category = "Training", meta = (ClampMin = "0.0", UIMin = "0.0", ClampMax = "1.0", UIMax = "1.0"))
-	float StudentForcingRate = 0.75f;
-
-	/** Number of denoising steps to use during training */
+	/** Number of denoising steps to use during training of distilled LODs. Larger values may improve quality marginally but makes training slower. */
 	UPROPERTY(EditAnywhere, Category = "Training", meta = (ClampMin = "1", UIMin = "1"))
 	int32 DenoiserSteps = 4;
 
diff --git a/Engine/Plugins/Experimental/Animation/AnimGen/Source/AnimGenEditor/Private/AnimGenEditorControllerToolkit.cpp b/Engine/Plugins/Experimental/Animation/AnimGen/Source/AnimGenEditor/Private/AnimGenEditorControllerToolkit.cpp
--- a/Engine/Plugins/Experimental/Animation/AnimGen/Source/AnimGenEditor/Private/AnimGenEditorControllerToolkit.cpp
+++ b/Engine/Plugins/Experimental/Animation/AnimGen/Source/AnimGenEditor/Private/AnimGenEditorControllerToolkit.cpp
@@ -1296,6 +1296,7 @@ namespace UE::AnimGen::Editor
 				ConfigObject->SetNumberField(TEXT("DenoiserHiddenUnitNum"), Controller->TrainingSettings->DenoiserHiddenUnitNum);
 
 				ConfigObject->SetNumberField(TEXT("IterationNum"), Controller->TrainingSettings->NumberOfIterations);
+				ConfigObject->SetNumberField(TEXT("IterationDistillNum"), Controller->TrainingSettings->NumberOfDistillationIterations);
 				ConfigObject->SetNumberField(TEXT("WarmupIterations"), Controller->TrainingSettings->WarmupIterations);
 				ConfigObject->SetNumberField(TEXT("LearningRate"), Controller->TrainingSettings->LearningRate);
 				ConfigObject->SetNumberField(TEXT("BatchSize"), Controller->TrainingSettings->BatchSize);
@@ -1305,7 +1306,6 @@ namespace UE::AnimGen::Editor
 				ConfigObject->SetNumberField(TEXT("MaxPoseNoiseScale"), Controller->TrainingSettings->MaxPoseNoiseScale);
 				ConfigObject->SetNumberField(TEXT("MinPoseNoiseScale"), Controller->TrainingSettings->MinPoseNoiseScale);
 				ConfigObject->SetNumberField(TEXT("ControlNoiseScale"), Controller->TrainingSettings->ControlNoiseScale);
-				ConfigObject->SetNumberField(TEXT("StudentForcingRate"), Controller->TrainingSettings->StudentForcingRate);
 
 				TArray<TSharedPtr<FJsonValue>> JSONRangeSequenceNames;
 				JSONRangeSequenceNames.Reserve(TotalRangeNum);
@@ -1367,7 +1367,7 @@ namespace UE::AnimGen::Editor
 
 				// Reset Iterations
 
-				MaxIterations = Controller->TrainingSettings->NumberOfIterations;
+				MaxIterations = Controller->TrainingSettings->NumberOfIterations + Controller->TrainingSettings->NumberOfDistillationIterations;
 				Iteration = 0;
 				StartTrainingTimeIteration = -1;
 
