AnimGen: Switched controls to using automatic normalization

#Synced-CL 51886592

diff --git a/Engine/Plugins/Experimental/Animation/AnimGen/Source/AnimGen/Private/AnimGenControl.cpp b/Engine/Plugins/Experimental/Animation/AnimGen/Source/AnimGen/Private/AnimGenControl.cpp
--- a/Engine/Plugins/Experimental/Animation/AnimGen/Source/AnimGen/Private/AnimGenControl.cpp
+++ b/Engine/Plugins/Experimental/Animation/AnimGen/Source/AnimGen/Private/AnimGenControl.cpp
@@ -716,7 +716,7 @@ FAnimGenControlSchemaElement UAnimGenControls::SpecifyNullControl(FAnimGenContro
 	}
 }
 
-FAnimGenControlSchemaElement UAnimGenControls::SpecifyContinuousControl(FAnimGenControlSchema& Schema, const int32 Size, const float Scale, const FName Tag)
+FAnimGenControlSchemaElement UAnimGenControls::SpecifyContinuousControl(FAnimGenControlSchema& Schema, const int32 Size, const FName Tag)
 {
 	if (!Schema.IsValid())
 	{
@@ -730,18 +730,9 @@ FAnimGenControlSchemaElement UAnimGenControls::SpecifyContinuousControl(FAnimGen
 		return FAnimGenControlSchemaElement();
 	}
 
-	if (Size == 0)
-	{
-		UE_ANIMGEN_LOG_ONCE(LogAnimGen, Warning, "SpecifyContinuousControl: Specifying zero-sized Continuous Control.");
-	}
-
-	TArray<float, TInlineAllocator<128>> Scales, Offsets;
-	Scales.Init(Scale, Size);
-	Offsets.Init(0.0f, Size);
-
 	{
 		FWriteScopeLock SchemaLock(Schema.ObservationSchema->Lock);
-		return { Schema.ObservationSchema->Schema.CreateContinuous({ Size, Scales, Offsets, UE::Learning::Observation::ENormalization::Manual }, Tag) };
+		return { Schema.ObservationSchema->Schema.CreateContinuous({ Size, {}, {}, UE::Learning::Observation::ENormalization::AutoDimensionAverage}, Tag) };
 	}
 }
 
@@ -852,7 +843,7 @@ FAnimGenControlSchemaElement UAnimGenControls::SpecifyInclusiveDiscreteControl(F
 
 FAnimGenControlSchemaElement UAnimGenControls::SpecifyCountControl(FAnimGenControlSchema& Schema, const FName Tag)
 {
-	return SpecifyContinuousControl(Schema, 1, 1.0f, Tag);
+	return SpecifyContinuousControl(Schema, 1, Tag);
 }
 
 FAnimGenControlSchemaElement UAnimGenControls::SpecifyStructControl(FAnimGenControlSchema& Schema, const TMap<FName, FAnimGenControlSchemaElement>& Elements, const FName Tag)
@@ -1360,27 +1351,27 @@ FAnimGenControlSchemaElement UAnimGenControls::SpecifyBoolControl(FAnimGenContro
 	return SpecifyNamedExclusiveDiscreteControl(Schema, { TEXT("true"), TEXT("false") }, Tag);
 }
 
-FAnimGenControlSchemaElement UAnimGenControls::SpecifyFloatControl(FAnimGenControlSchema& Schema, const float FloatScale, const FName Tag)
+FAnimGenControlSchemaElement UAnimGenControls::SpecifyFloatControl(FAnimGenControlSchema& Schema, const FName Tag)
 {
-	return SpecifyContinuousControl(Schema, 1, FloatScale, Tag);
+	return SpecifyContinuousControl(Schema, 1, Tag);
 }
 
-FAnimGenControlSchemaElement UAnimGenControls::SpecifyLocationControl(FAnimGenControlSchema& Schema, const float LocationScale, const FName Tag)
+FAnimGenControlSchemaElement UAnimGenControls::SpecifyLocationControl(FAnimGenControlSchema& Schema, const FName Tag)
 {
-	return SpecifyContinuousControl(Schema, 3, LocationScale, Tag);
+	return SpecifyContinuousControl(Schema, 3, Tag);
 }
 
 FAnimGenControlSchemaElement UAnimGenControls::SpecifyRotationControl(FAnimGenControlSchema& Schema, const FName Tag)
 {
-	return SpecifyContinuousControl(Schema, 6, 1.0f, Tag);
+	return SpecifyContinuousControl(Schema, 6, Tag);
 }
 
 FAnimGenControlSchemaElement UAnimGenControls::SpecifyScaleControl(FAnimGenControlSchema& Schema, const FName Tag)
 {
-	return SpecifyContinuousControl(Schema, 3, 1.0f, Tag);
+	return SpecifyContinuousControl(Schema, 3, Tag);
 }
 
-FAnimGenControlSchemaElement UAnimGenControls::SpecifyTransformControl(FAnimGenControlSchema& Schema, const float LocationScale, const FName Tag)
+FAnimGenControlSchemaElement UAnimGenControls::SpecifyTransformControl(FAnimGenControlSchema& Schema, const FName Tag)
 {
 	return SpecifyStructControlFromArrayViews(Schema,
 		{
@@ -1389,7 +1380,7 @@ FAnimGenControlSchemaElement UAnimGenControls::SpecifyTransformControl(FAnimGenC
 			TEXT("Scale")
 		},
 		{
-			SpecifyLocationControl(Schema, LocationScale),
+			SpecifyLocationControl(Schema),
 			SpecifyRotationControl(Schema),
 			SpecifyScaleControl(Schema)
 		},
@@ -1398,37 +1389,37 @@ FAnimGenControlSchemaElement UAnimGenControls::SpecifyTransformControl(FAnimGenC
 
 FAnimGenControlSchemaElement UAnimGenControls::SpecifyAngleControl(FAnimGenControlSchema& Schema, const FName Tag)
 {
-	return SpecifyContinuousControl(Schema, 2, 1.0f, Tag);
+	return SpecifyContinuousControl(Schema, 2, Tag);
 }
 
-FAnimGenControlSchemaElement UAnimGenControls::SpecifyLinearVelocityControl(FAnimGenControlSchema& Schema, const float LinearVelocityScale, const FName Tag)
+FAnimGenControlSchemaElement UAnimGenControls::SpecifyLinearVelocityControl(FAnimGenControlSchema& Schema, const FName Tag)
 {
-	return SpecifyContinuousControl(Schema, 3, LinearVelocityScale, Tag);
+	return SpecifyContinuousControl(Schema, 3, Tag);
 }
 
-FAnimGenControlSchemaElement UAnimGenControls::SpecifyAngularVelocityControl(FAnimGenControlSchema& Schema, const float AngularVelocityScale, const FName Tag)
+FAnimGenControlSchemaElement UAnimGenControls::SpecifyAngularVelocityControl(FAnimGenControlSchema& Schema, const FName Tag)
 {
-	return SpecifyContinuousControl(Schema, 3, AngularVelocityScale, Tag);
+	return SpecifyContinuousControl(Schema, 3, Tag);
 }
 
 FAnimGenControlSchemaElement UAnimGenControls::SpecifyDirectionControl(FAnimGenControlSchema& Schema, const FName Tag)
 {
-	return SpecifyContinuousControl(Schema, 3, 1.0f, Tag);
+	return SpecifyContinuousControl(Schema, 3, Tag);
 }
 
-FAnimGenControlSchemaElement UAnimGenControls::SpecifyTimeControl(FAnimGenControlSchema& Schema, const float TimeScale, const FName Tag)
+FAnimGenControlSchemaElement UAnimGenControls::SpecifyTimeControl(FAnimGenControlSchema& Schema, const FName Tag)
 {
-	return SpecifyContinuousControl(Schema, 1, TimeScale, Tag);
+	return SpecifyContinuousControl(Schema, 1, Tag);
 }
 
-FAnimGenControlSchemaElement UAnimGenControls::SpecifyLocationsStaticArrayControl(FAnimGenControlSchema& Schema, const int32 Num, const float LocationScale, const FName Tag)
+FAnimGenControlSchemaElement UAnimGenControls::SpecifyLocationsStaticArrayControl(FAnimGenControlSchema& Schema, const int32 Num, const FName Tag)
 {
-	return SpecifyStaticArrayControl(Schema, SpecifyLocationControl(Schema, LocationScale), Num, Tag);
+	return SpecifyStaticArrayControl(Schema, SpecifyLocationControl(Schema), Num, Tag);
 }
 
-FAnimGenControlSchemaElement UAnimGenControls::SpecifyLinearVelocitiesStaticArrayControl(FAnimGenControlSchema& Schema, const int32 Num, const float LinearVelocityScale, const FName Tag)
+FAnimGenControlSchemaElement UAnimGenControls::SpecifyLinearVelocitiesStaticArrayControl(FAnimGenControlSchema& Schema, const int32 Num, const FName Tag)
 {
-	return SpecifyStaticArrayControl(Schema, SpecifyLinearVelocityControl(Schema, LinearVelocityScale), Num, Tag);
+	return SpecifyStaticArrayControl(Schema, SpecifyLinearVelocityControl(Schema), Num, Tag);
 }
 
 FAnimGenControlSchemaElement UAnimGenControls::SpecifyDirectionsStaticArrayControl(FAnimGenControlSchema& Schema, const int32 Num, const FName Tag)
@@ -1436,7 +1427,7 @@ FAnimGenControlSchemaElement UAnimGenControls::SpecifyDirectionsStaticArrayContr
 	return SpecifyStaticArrayControl(Schema, SpecifyDirectionControl(Schema), Num, Tag);
 }
 
-FAnimGenControlSchemaElement UAnimGenControls::SpecifyPoseControl(FAnimGenControlSchema& Schema, const int32 BoneNum, const float LocationScale, const float LinearVelocityScale, const FName Tag)
+FAnimGenControlSchemaElement UAnimGenControls::SpecifyPoseControl(FAnimGenControlSchema& Schema, const int32 BoneNum, const FName Tag)
 {
 	return SpecifyStructControlFromArrayViews(
 		Schema,
@@ -1445,13 +1436,13 @@ FAnimGenControlSchemaElement UAnimGenControls::SpecifyPoseControl(FAnimGenContro
 			TEXT("LinearVelocities"),
 		},
 		{
-			SpecifyLocationsStaticArrayControl(Schema, BoneNum, LocationScale),
-			SpecifyLinearVelocitiesStaticArrayControl(Schema, BoneNum, LinearVelocityScale),
+			SpecifyLocationsStaticArrayControl(Schema, BoneNum),
+			SpecifyLinearVelocitiesStaticArrayControl(Schema, BoneNum),
 		},
 		Tag);
 }
 
-FAnimGenControlSchemaElement UAnimGenControls::SpecifyTrajectoryControl(FAnimGenControlSchema& Schema, const int32 TrajectorySampleNum, const float LocationScale, const FName Tag)
+FAnimGenControlSchemaElement UAnimGenControls::SpecifyTrajectoryControl(FAnimGenControlSchema& Schema, const int32 TrajectorySampleNum, const FName Tag)
 {
 	return SpecifyStructControlFromArrayViews(
 		Schema,
@@ -1460,20 +1451,20 @@ FAnimGenControlSchemaElement UAnimGenControls::SpecifyTrajectoryControl(FAnimGen
 			TEXT("Directions"),
 		},
 		{
-			SpecifyLocationsStaticArrayControl(Schema, TrajectorySampleNum, LocationScale),
+			SpecifyLocationsStaticArrayControl(Schema, TrajectorySampleNum),
 			SpecifyDirectionsStaticArrayControl(Schema, TrajectorySampleNum),
 		},
 		Tag);
 }
 
-FAnimGenControlSchemaElement UAnimGenControls::SpecifyLocationTrajectoryControl(FAnimGenControlSchema& Schema, const int32 TrajectorySampleNum, const float LocationScale, const FName Tag)
+FAnimGenControlSchemaElement UAnimGenControls::SpecifyLocationTrajectoryControl(FAnimGenControlSchema& Schema, const int32 TrajectorySampleNum, const FName Tag)
 {
-	return SpecifyLocationsStaticArrayControl(Schema, TrajectorySampleNum, LocationScale, Tag);
+	return SpecifyLocationsStaticArrayControl(Schema, TrajectorySampleNum, Tag);
 }
 
 FAnimGenControlSchemaElement UAnimGenControls::SpecifyEventControl(FAnimGenControlSchema& Schema, const FName Tag)
 {
-	return SpecifyContinuousControl(Schema, 2, 1.0f, Tag);
+	return SpecifyContinuousControl(Schema, 2, Tag);
 }
 
 
diff --git a/Engine/Plugins/Experimental/Animation/AnimGen/Source/AnimGen/Public/AnimGenControl.h b/Engine/Plugins/Experimental/Animation/AnimGen/Source/AnimGen/Public/AnimGenControl.h
--- a/Engine/Plugins/Experimental/Animation/AnimGen/Source/AnimGen/Public/AnimGenControl.h
+++ b/Engine/Plugins/Experimental/Animation/AnimGen/Source/AnimGen/Public/AnimGenControl.h
@@ -238,8 +238,8 @@ public:
 	UFUNCTION(BlueprintPure, Category = "AnimGen", meta = (AdvancedDisplay = 1))
 	static UE_API FAnimGenControlSchemaElement SpecifyNullControl(UPARAM(ref) FAnimGenControlSchema& Schema, const FName Tag = TEXT("NullControl"));
 
-	UFUNCTION(BlueprintPure, Category = "AnimGen", meta = (AdvancedDisplay = 3))
-	static UE_API FAnimGenControlSchemaElement SpecifyContinuousControl(UPARAM(ref) FAnimGenControlSchema& Schema, const int32 Size, const float Scale = 1.0f, const FName Tag = TEXT("ContinuousControl"));
+	UFUNCTION(BlueprintPure, Category = "AnimGen", meta = (AdvancedDisplay = 2))
+	static UE_API FAnimGenControlSchemaElement SpecifyContinuousControl(UPARAM(ref) FAnimGenControlSchema& Schema, const int32 Size, const FName Tag = TEXT("ContinuousControl"));
 
 	UFUNCTION(BlueprintPure, Category = "AnimGen", meta = (AdvancedDisplay = 2))
 	static UE_API FAnimGenControlSchemaElement SpecifyNamedExclusiveDiscreteControl(UPARAM(ref) FAnimGenControlSchema& Schema, const TArray<FName>& ElementNames, const FName Tag = TEXT("NamedExclusiveDiscreteControl"));
@@ -313,11 +313,11 @@ public:
 	UFUNCTION(BlueprintPure, Category = "AnimGen", meta = (AdvancedDisplay = 1))
 	static UE_API FAnimGenControlSchemaElement SpecifyBoolControl(UPARAM(ref) FAnimGenControlSchema& Schema, const FName Tag = TEXT("BoolControl"));
 
-	UFUNCTION(BlueprintPure, Category = "AnimGen", meta = (AdvancedDisplay = 2))
-	static UE_API FAnimGenControlSchemaElement SpecifyFloatControl(UPARAM(ref) FAnimGenControlSchema& Schema, const float FloatScale = 1.0f, const FName Tag = TEXT("FloatControl"));
+	UFUNCTION(BlueprintPure, Category = "AnimGen", meta = (AdvancedDisplay = 1))
+	static UE_API FAnimGenControlSchemaElement SpecifyFloatControl(UPARAM(ref) FAnimGenControlSchema& Schema, const FName Tag = TEXT("FloatControl"));
 
-	UFUNCTION(BlueprintPure, Category = "AnimGen", meta = (AdvancedDisplay = 2))
-	static UE_API FAnimGenControlSchemaElement SpecifyLocationControl(UPARAM(ref) FAnimGenControlSchema& Schema, const float LocationScale = 100.0f, const FName Tag = TEXT("LocationControl"));
+	UFUNCTION(BlueprintPure, Category = "AnimGen", meta = (AdvancedDisplay = 1))
+	static UE_API FAnimGenControlSchemaElement SpecifyLocationControl(UPARAM(ref) FAnimGenControlSchema& Schema, const FName Tag = TEXT("LocationControl"));
 
 	UFUNCTION(BlueprintPure, Category = "AnimGen", meta = (AdvancedDisplay = 1))
 	static UE_API FAnimGenControlSchemaElement SpecifyRotationControl(UPARAM(ref) FAnimGenControlSchema& Schema, const FName Tag = TEXT("RotationControl"));
@@ -325,43 +325,43 @@ public:
 	UFUNCTION(BlueprintPure, Category = "AnimGen", meta = (AdvancedDisplay = 1))
 	static UE_API FAnimGenControlSchemaElement SpecifyScaleControl(UPARAM(ref) FAnimGenControlSchema& Schema, const FName Tag = TEXT("ScaleControl"));
 
-	UFUNCTION(BlueprintPure, Category = "AnimGen", meta = (AdvancedDisplay = 2))
-	static UE_API FAnimGenControlSchemaElement SpecifyTransformControl(UPARAM(ref) FAnimGenControlSchema& Schema, const float LocationScale = 100.0f, const FName Tag = TEXT("TransformControl"));
+	UFUNCTION(BlueprintPure, Category = "AnimGen", meta = (AdvancedDisplay = 1))
+	static UE_API FAnimGenControlSchemaElement SpecifyTransformControl(UPARAM(ref) FAnimGenControlSchema& Schema, const FName Tag = TEXT("TransformControl"));
 
 	UFUNCTION(BlueprintPure, Category = "AnimGen", meta = (AdvancedDisplay = 1))
 	static UE_API FAnimGenControlSchemaElement SpecifyAngleControl(UPARAM(ref) FAnimGenControlSchema& Schema, const FName Tag = TEXT("AngleControl"));
 
-	UFUNCTION(BlueprintPure, Category = "AnimGen", meta = (AdvancedDisplay = 2))
-	static UE_API FAnimGenControlSchemaElement SpecifyLinearVelocityControl(UPARAM(ref) FAnimGenControlSchema& Schema, const float LinearVelocityScale = 200.0f, const FName Tag = TEXT("LinearVelocityControl"));
+	UFUNCTION(BlueprintPure, Category = "AnimGen", meta = (AdvancedDisplay = 1))
+	static UE_API FAnimGenControlSchemaElement SpecifyLinearVelocityControl(UPARAM(ref) FAnimGenControlSchema& Schema, const FName Tag = TEXT("LinearVelocityControl"));
 
-	UFUNCTION(BlueprintPure, Category = "AnimGen", meta = (AdvancedDisplay = 2))
-	static UE_API FAnimGenControlSchemaElement SpecifyAngularVelocityControl(UPARAM(ref) FAnimGenControlSchema& Schema, const float AngularVelocityScale = 180.0f, const FName Tag = TEXT("AngularVelocityControl"));
+	UFUNCTION(BlueprintPure, Category = "AnimGen", meta = (AdvancedDisplay = 1))
+	static UE_API FAnimGenControlSchemaElement SpecifyAngularVelocityControl(UPARAM(ref) FAnimGenControlSchema& Schema, const FName Tag = TEXT("AngularVelocityControl"));
 
 	UFUNCTION(BlueprintPure, Category = "AnimGen", meta = (AdvancedDisplay = 1))
 	static UE_API FAnimGenControlSchemaElement SpecifyDirectionControl(UPARAM(ref) FAnimGenControlSchema& Schema, const FName Tag = TEXT("DirectionControl"));
 
-	UFUNCTION(BlueprintPure, Category = "AnimGen", meta = (AdvancedDisplay = 2))
-	static UE_API FAnimGenControlSchemaElement SpecifyTimeControl(UPARAM(ref) FAnimGenControlSchema& Schema, const float TimeScale = 1.0f, const FName Tag = TEXT("TimeControl"));
+	UFUNCTION(BlueprintPure, Category = "AnimGen", meta = (AdvancedDisplay = 1))
+	static UE_API FAnimGenControlSchemaElement SpecifyTimeControl(UPARAM(ref) FAnimGenControlSchema& Schema, const FName Tag = TEXT("TimeControl"));
 
 
 	UFUNCTION(BlueprintPure, Category = "AnimGen", meta = (AdvancedDisplay = 2))
-	static UE_API FAnimGenControlSchemaElement SpecifyLocationsStaticArrayControl(UPARAM(ref) FAnimGenControlSchema& Schema, const int32 Num, const float LocationScale = 100.0f, const FName Tag = TEXT("LocationsStaticArrayControl"));
+	static UE_API FAnimGenControlSchemaElement SpecifyLocationsStaticArrayControl(UPARAM(ref) FAnimGenControlSchema& Schema, const int32 Num, const FName Tag = TEXT("LocationsStaticArrayControl"));
 
 	UFUNCTION(BlueprintPure, Category = "AnimGen", meta = (AdvancedDisplay = 2))
-	static UE_API FAnimGenControlSchemaElement SpecifyLinearVelocitiesStaticArrayControl(UPARAM(ref) FAnimGenControlSchema& Schema, const int32 Num, const float LinearVelocityScale = 200.0f, const FName Tag = TEXT("LinearVelocitiesStaticArrayControl"));
+	static UE_API FAnimGenControlSchemaElement SpecifyLinearVelocitiesStaticArrayControl(UPARAM(ref) FAnimGenControlSchema& Schema, const int32 Num, const FName Tag = TEXT("LinearVelocitiesStaticArrayControl"));
 
 	UFUNCTION(BlueprintPure, Category = "AnimGen", meta = (AdvancedDisplay = 2))
 	static UE_API FAnimGenControlSchemaElement SpecifyDirectionsStaticArrayControl(UPARAM(ref) FAnimGenControlSchema& Schema, const int32 Num, const FName Tag = TEXT("DirectionsStaticArrayControl"));
 
 
-	UFUNCTION(BlueprintPure, Category = "AnimGen", meta = (AdvancedDisplay = 4))
-	static UE_API FAnimGenControlSchemaElement SpecifyPoseControl(UPARAM(ref) FAnimGenControlSchema& Schema, const int32 BoneNum, const float LocationScale = 100.0f, const float LinearVelocityScale = 200.0f, const FName Tag = TEXT("PoseControl"));
+	UFUNCTION(BlueprintPure, Category = "AnimGen", meta = (AdvancedDisplay = 2))
+	static UE_API FAnimGenControlSchemaElement SpecifyPoseControl(UPARAM(ref) FAnimGenControlSchema& Schema, const int32 BoneNum, const FName Tag = TEXT("PoseControl"));
 
-	UFUNCTION(BlueprintPure, Category = "AnimGen", meta = (AdvancedDisplay = 3))
-	static UE_API FAnimGenControlSchemaElement SpecifyTrajectoryControl(UPARAM(ref) FAnimGenControlSchema& Schema, const int32 TrajectorySampleNum, const float LocationScale = 100.0f, const FName Tag = TEXT("TrajectoryControl"));
+	UFUNCTION(BlueprintPure, Category = "AnimGen", meta = (AdvancedDisplay = 2))
+	static UE_API FAnimGenControlSchemaElement SpecifyTrajectoryControl(UPARAM(ref) FAnimGenControlSchema& Schema, const int32 TrajectorySampleNum, const FName Tag = TEXT("TrajectoryControl"));
 
-	UFUNCTION(BlueprintPure, Category = "AnimGen", meta = (AdvancedDisplay = 3))
-	static UE_API FAnimGenControlSchemaElement SpecifyLocationTrajectoryControl(UPARAM(ref) FAnimGenControlSchema& Schema, const int32 TrajectorySampleNum, const float LocationScale = 100.0f, const FName Tag = TEXT("LocationTrajectoryControl"));
+	UFUNCTION(BlueprintPure, Category = "AnimGen", meta = (AdvancedDisplay = 2))
+	static UE_API FAnimGenControlSchemaElement SpecifyLocationTrajectoryControl(UPARAM(ref) FAnimGenControlSchema& Schema, const int32 TrajectorySampleNum, const FName Tag = TEXT("LocationTrajectoryControl"));
 
 
 	UFUNCTION(BlueprintPure, Category = "AnimGen", meta = (AdvancedDisplay = 1))
diff --git a/Engine/Plugins/Experimental/LearningCore/Source/Learning/Public/LearningObservation.h b/Engine/Plugins/Experimental/LearningCore/Source/Learning/Public/LearningObservation.h
--- a/Engine/Plugins/Experimental/LearningCore/Source/Learning/Public/LearningObservation.h
+++ b/Engine/Plugins/Experimental/LearningCore/Source/Learning/Public/LearningObservation.h
@@ -71,6 +71,7 @@ namespace UE::Learning::Observation
 		Manual = 0,
 		AutoShared = 1,
 		AutoPerDimension = 2,
+		AutoDimensionAverage = 3,
 	};
 
 	static inline FString GetNormalizationOptionAsString(const ENormalization NormalizationOption)
@@ -80,6 +81,7 @@ namespace UE::Learning::Observation
 		case UE::Learning::Observation::ENormalization::Manual: return TEXT("Manual");
 		case UE::Learning::Observation::ENormalization::AutoShared: return TEXT("AutoShared");
 		case UE::Learning::Observation::ENormalization::AutoPerDimension: return TEXT("AutoPerDimension");
+		case UE::Learning::Observation::ENormalization::AutoDimensionAverage: return TEXT("AutoDimensionAverage");
 		default: checkNoEntry(); return TEXT("");
 		}
 	}
