AnimGen: Switched controller node to use Instanced UObjects for AnimNotifies and AnimNotifyStates output.

#Synced-CL 52006914

diff --git a/Engine/Plugins/Experimental/Animation/AnimGen/Source/AnimGen/Private/AnimNode_AnimGenController.cpp b/Engine/Plugins/Experimental/Animation/AnimGen/Source/AnimGen/Private/AnimNode_AnimGenController.cpp
--- a/Engine/Plugins/Experimental/Animation/AnimGen/Source/AnimGen/Private/AnimNode_AnimGenController.cpp
+++ b/Engine/Plugins/Experimental/Animation/AnimGen/Source/AnimGen/Private/AnimNode_AnimGenController.cpp
@@ -526,7 +526,7 @@ void FAnimNode_AnimGenController::Update_AnyThread(const FAnimationUpdateContext
 
 	AnimNotifyEvents.Reset();
 
-	for (const TPair<FName, TSubclassOf<UAnimNotify>>& AttributeAnimNotify : AnimNotifyMap)
+	for (const TPair<FName, TObjectPtr<UAnimNotify>>& AttributeAnimNotify : AnimNotifyMap)
 	{
 		if (!AttributeAnimNotify.Value) { continue; }
 
@@ -567,12 +567,12 @@ void FAnimNode_AnimGenController::Update_AnyThread(const FAnimationUpdateContext
 			FMath::Abs(PrevTimeUntilEvent - CurrTimeUntilEvent) < 0.1f)
 		{
 			FAnimNotifyEvent& NotifyEvent = AnimNotifyEvents.AddDefaulted_GetRef();
-			NotifyEvent.Notify = AttributeAnimNotify.Value->GetDefaultObject<UAnimNotify>();
+			NotifyEvent.Notify = AttributeAnimNotify.Value;
 			NotifyEvent.NotifyName = AutoEncoder->GetAttributeName(AttributeIdx);
 		}
 	}
 
-	for (const TPair<FName, TSubclassOf<UAnimNotifyState>>& AttributeAnimNotifyState : AnimNotifyStateMap)
+	for (const TPair<FName, TObjectPtr<UAnimNotifyState>>& AttributeAnimNotifyState : AnimNotifyStateMap)
 	{
 		if (!AttributeAnimNotifyState.Value) { continue; }
 
@@ -601,7 +601,7 @@ void FAnimNode_AnimGenController::Update_AnyThread(const FAnimationUpdateContext
 		if (PoseState.CurrPoseData->AttributeData.GetBool(0, AttributeIdx))
 		{
 			FAnimNotifyEvent& NotifyEvent = AnimNotifyEvents.AddDefaulted_GetRef();
-			NotifyEvent.NotifyStateClass = AttributeAnimNotifyState.Value->GetDefaultObject<UAnimNotifyState>();
+			NotifyEvent.NotifyStateClass = AttributeAnimNotifyState.Value;
 			NotifyEvent.NotifyName = AutoEncoder->GetAttributeName(AttributeIdx);
 		}
 	}
diff --git a/Engine/Plugins/Experimental/Animation/AnimGen/Source/AnimGen/Public/AnimNode_AnimGenController.h b/Engine/Plugins/Experimental/Animation/AnimGen/Source/AnimGen/Public/AnimNode_AnimGenController.h
--- a/Engine/Plugins/Experimental/Animation/AnimGen/Source/AnimGen/Public/AnimNode_AnimGenController.h
+++ b/Engine/Plugins/Experimental/Animation/AnimGen/Source/AnimGen/Public/AnimNode_AnimGenController.h
@@ -70,12 +70,12 @@ public:
 public:
 
 	/** Map of attribute names to anim notify states to produce */
-	UPROPERTY(EditAnywhere, Category = Notifies)
-	TMap<FName, TSubclassOf<UAnimNotifyState>> AnimNotifyStateMap;
+	UPROPERTY(EditAnywhere, Instanced, Category = Notifies)
+	TMap<FName, TObjectPtr<UAnimNotifyState>> AnimNotifyStateMap;
 
 	/** Map of attribute names to anim notifies to produce */
-	UPROPERTY(EditAnywhere, Category = Notifies)
-	TMap<FName, TSubclassOf<UAnimNotify>> AnimNotifyMap;
+	UPROPERTY(EditAnywhere, Instanced, Category = Notifies)
+	TMap<FName, TObjectPtr<UAnimNotify>> AnimNotifyMap;
 
 	/** Map of attribute names to curves to produce */
 	UPROPERTY(EditAnywhere, Category = Curves)
