AnimGen: Added support for past time in trajectory

#Synced-CL 52204516

diff --git a/Engine/Plugins/Experimental/Animation/AnimDatabase/Source/AnimDatabase/Private/AnimDatabaseFrameAttribute.cpp b/Engine/Plugins/Experimental/Animation/AnimDatabase/Source/AnimDatabase/Private/AnimDatabaseFrameAttribute.cpp
--- a/Engine/Plugins/Experimental/Animation/AnimDatabase/Source/AnimDatabase/Private/AnimDatabaseFrameAttribute.cpp
+++ b/Engine/Plugins/Experimental/Animation/AnimDatabase/Source/AnimDatabase/Private/AnimDatabaseFrameAttribute.cpp
@@ -1821,6 +1821,7 @@ void UAnimDatabaseFrameAttributeLibrary::MakeTrajectoryFrameAttribute(
 	const FAnimDatabaseFrameRanges& FrameRanges,
 	const int32 SampleNum,
 	const float FutureTime,
+	const float PastTime,
 	const FVector ForwardVector)
 {
 	if (SampleNum < 0)
@@ -1835,12 +1836,12 @@ void UAnimDatabaseFrameAttributeLibrary::MakeTrajectoryFrameAttribute(
 	OutDirectionFrameAttributes.SetNum(SampleNum);
 
 	ParallelFor(SampleNum, [
-		SampleNum, FutureTime, ForwardVector, Database,
+		SampleNum, FutureTime, PastTime, ForwardVector, Database,
 		&OutLocationFrameAttributes,
 		&OutDirectionFrameAttributes,
 		&FrameRanges](const int32 SampleIdx) {
 
-			const float RelativeTime = (FutureTime / FMath::Max(SampleNum - 1, 1)) * SampleIdx;
+			const float RelativeTime = ((FutureTime + PastTime) / FMath::Max(SampleNum - 1, 1)) * SampleIdx - PastTime;
 
 			MakeRootLocationAndDirectionFrameAttribute(
 				OutLocationFrameAttributes[SampleIdx],
@@ -1855,7 +1856,8 @@ void UAnimDatabaseFrameAttributeLibrary::MakeLocationTrajectoryFrameAttribute(
 	const UAnimDatabase* Database,
 	const FAnimDatabaseFrameRanges& FrameRanges,
 	const int32 SampleNum,
-	const float FutureTime)
+	const float FutureTime,
+	const float PastTime)
 {
 	if (SampleNum < 0)
 	{
@@ -1867,11 +1869,11 @@ void UAnimDatabaseFrameAttributeLibrary::MakeLocationTrajectoryFrameAttribute(
 	OutLocationFrameAttributes.SetNum(SampleNum);
 
 	ParallelFor(SampleNum, [
-		SampleNum, FutureTime, Database,
+		SampleNum, FutureTime, PastTime, Database,
 		&OutLocationFrameAttributes,
 		&FrameRanges](const int32 SampleIdx) {
 
-			const float RelativeTime = (FutureTime / FMath::Max(SampleNum - 1, 1)) * SampleIdx;
+			const float RelativeTime = ((FutureTime + PastTime) / FMath::Max(SampleNum - 1, 1)) * SampleIdx - PastTime;
 			OutLocationFrameAttributes[SampleIdx] = MakeRootLocationFrameAttribute(Database, FrameRanges, RelativeTime);
 
 		});
@@ -1885,6 +1887,7 @@ void UAnimDatabaseFrameAttributeLibrary::MakeLookAtTrajectoryFrameAttribute(
 	const int32 LookAtBoneIndex,
 	const int32 SampleNum,
 	const float FutureTime,
+	const float PastTime,
 	const FVector LocalDirection,
 	const float DirectionSmoothingAmount)
 {
@@ -1900,12 +1903,12 @@ void UAnimDatabaseFrameAttributeLibrary::MakeLookAtTrajectoryFrameAttribute(
 	OutDirectionFrameAttributes.SetNum(SampleNum);
 
 	ParallelFor(SampleNum, [
-		SampleNum, FutureTime, Database, LookAtBoneIndex, LocalDirection, DirectionSmoothingAmount,
+		SampleNum, FutureTime, PastTime, Database, LookAtBoneIndex, LocalDirection, DirectionSmoothingAmount,
 		&OutLocationFrameAttributes,
 		&OutDirectionFrameAttributes,
 		&FrameRanges](const int32 SampleIdx) {
 
-			const float RelativeTime = (FutureTime / FMath::Max(SampleNum - 1, 1)) * SampleIdx;
+			const float RelativeTime = ((FutureTime + PastTime) / FMath::Max(SampleNum - 1, 1)) * SampleIdx - PastTime;
 
 			OutLocationFrameAttributes[SampleIdx] = MakeRootLocationFrameAttribute(Database, FrameRanges, RelativeTime);
 			OutDirectionFrameAttributes[SampleIdx] = FrameAttributeFilterGaussian(FrameAttributeProject(
@@ -1922,13 +1925,14 @@ void UAnimDatabaseFrameAttributeLibrary::MakeRootLocalTrajectoryFrameAttributeAn
 	const FAnimDatabaseFrameRanges& FrameRanges,
 	const int32 SampleNum,
 	const float FutureTime,
+	const float PastTime,
 	const FVector ForwardVector)
 {
 	FAnimDatabaseFrameAttribute RootLocation, RootRotation;
 	MakeRootLocationAndRotationFrameAttribute(RootLocation, RootRotation, Database, FrameRanges);
 
 	TArray<FAnimDatabaseFrameAttribute> LocationFrameAttributes, DirectionFrameAttributes;
-	MakeTrajectoryFrameAttribute(LocationFrameAttributes, DirectionFrameAttributes, Database, FrameRanges, SampleNum, FutureTime, ForwardVector);
+	MakeTrajectoryFrameAttribute(LocationFrameAttributes, DirectionFrameAttributes, Database, FrameRanges, SampleNum, FutureTime, PastTime, ForwardVector);
 
 	FrameAttributesSubtractAndUnrotate(OutLocationFrameAttributes, LocationFrameAttributes, RootLocation, RootRotation);
 	FrameAttributesUnrotate(OutDirectionFrameAttributes, RootRotation, DirectionFrameAttributes);
@@ -1944,6 +1948,7 @@ void UAnimDatabaseFrameAttributeLibrary::MakeRootLocalLookAtTrajectoryFrameAttri
 	const int32 LookAtBoneIndex,
 	const int32 SampleNum,
 	const float FutureTime,
+	const float PastTime,
 	const FVector LocalDirection,
 	const float DirectionSmoothingAmount)
 {
@@ -1951,7 +1956,7 @@ void UAnimDatabaseFrameAttributeLibrary::MakeRootLocalLookAtTrajectoryFrameAttri
 	MakeRootLocationAndRotationFrameAttribute(RootLocation, RootRotation, Database, FrameRanges);
 
 	TArray<FAnimDatabaseFrameAttribute> LocationFrameAttributes, DirectionFrameAttributes;
-	MakeLookAtTrajectoryFrameAttribute(LocationFrameAttributes, DirectionFrameAttributes, Database, FrameRanges, LookAtBoneIndex, SampleNum, FutureTime, LocalDirection);
+	MakeLookAtTrajectoryFrameAttribute(LocationFrameAttributes, DirectionFrameAttributes, Database, FrameRanges, LookAtBoneIndex, SampleNum, FutureTime, PastTime, LocalDirection);
 
 	FrameAttributesSubtractAndUnrotate(OutLocationFrameAttributes, LocationFrameAttributes, RootLocation, RootRotation);
 	FrameAttributesUnrotate(OutDirectionFrameAttributes, RootRotation, DirectionFrameAttributes);
diff --git a/Engine/Plugins/Experimental/Animation/AnimDatabase/Source/AnimDatabase/Public/AnimDatabaseFrameAttribute.h b/Engine/Plugins/Experimental/Animation/AnimDatabase/Source/AnimDatabase/Public/AnimDatabaseFrameAttribute.h
--- a/Engine/Plugins/Experimental/Animation/AnimDatabase/Source/AnimDatabase/Public/AnimDatabaseFrameAttribute.h
+++ b/Engine/Plugins/Experimental/Animation/AnimDatabase/Source/AnimDatabase/Public/AnimDatabaseFrameAttribute.h
@@ -479,7 +479,7 @@ public:
 public:
 	
 	/** Makes an array of location and direction frame attributes for the root at the given sample times into the future */
-	UFUNCTION(BlueprintPure = true, Category = "AnimDatabase", meta = (NotBlueprintThreadSafe, AdvancedDisplay = 7))
+	UFUNCTION(BlueprintPure = true, Category = "AnimDatabase", meta = (NotBlueprintThreadSafe, AdvancedDisplay = 8))
 	static UE_API void MakeTrajectoryFrameAttribute(
 		TArray<FAnimDatabaseFrameAttribute>& OutLocationFrameAttributes,
 		TArray<FAnimDatabaseFrameAttribute>& OutDirectionFrameAttributes,
@@ -487,22 +487,24 @@ public:
 		const FAnimDatabaseFrameRanges& FrameRanges,
 		const int32 SampleNum = 4,
 		const float FutureTime = 1.0f,
+		const float PastTime = 0.0f,
 		const FVector ForwardVector = FVector(0.0f, 1.0f, 0.0f));
 
 	/** Makes an array of location frame attributes for the root at the given sample times into the future */
-	UFUNCTION(BlueprintPure = true, Category = "AnimDatabase", meta = (NotBlueprintThreadSafe, AdvancedDisplay = 5))
+	UFUNCTION(BlueprintPure = true, Category = "AnimDatabase", meta = (NotBlueprintThreadSafe, AdvancedDisplay = 6))
 	static UE_API void MakeLocationTrajectoryFrameAttribute(
 		TArray<FAnimDatabaseFrameAttribute>& OutLocationFrameAttributes,
 		const UAnimDatabase* Database,
 		const FAnimDatabaseFrameRanges& FrameRanges,
 		const int32 SampleNum = 4,
-		const float FutureTime = 1.0f);
+		const float FutureTime = 1.0f,
+		const float PastTime = 0.0f);
 
 	/**
 	 * Makes an array of location and direction frame attributes for the root location and look-at-direction of the provided bone at the given 
 	 * sample times into the future
 	 */
-	UFUNCTION(BlueprintPure = true, Category = "AnimDatabase", meta = (NotBlueprintThreadSafe, AdvancedDisplay = 9))
+	UFUNCTION(BlueprintPure = true, Category = "AnimDatabase", meta = (NotBlueprintThreadSafe, AdvancedDisplay = 10))
 	static UE_API void MakeLookAtTrajectoryFrameAttribute(
 		TArray<FAnimDatabaseFrameAttribute>& OutLocationFrameAttributes,
 		TArray<FAnimDatabaseFrameAttribute>& OutDirectionFrameAttributes,
@@ -511,6 +513,7 @@ public:
 		const int32 LookAtBoneIndex,
 		const int32 SampleNum = 4,
 		const float FutureTime = 1.0f,
+		const float PastTime = 0.0f,
 		const FVector LocalDirection = FVector(0.0f, 1.0f, 0.0f),
 		const float DirectionSmoothingAmount = 15.0f);
 
@@ -518,7 +521,7 @@ public:
 	 * Makes trajectory frame attributes transformed local to the character's root location and rotation. Computes the overall scale of locations -
 	 * which is useful for scaling distance costs for use with matching.
 	 */
-	UFUNCTION(BlueprintPure = true, Category = "AnimDatabase", meta = (NotBlueprintThreadSafe, AdvancedDisplay = 8))
+	UFUNCTION(BlueprintPure = true, Category = "AnimDatabase", meta = (NotBlueprintThreadSafe, AdvancedDisplay = 9))
 	static UE_API void MakeRootLocalTrajectoryFrameAttributeAndScale(
 		TArray<FAnimDatabaseFrameAttribute>& OutLocationFrameAttributes,
 		TArray<FAnimDatabaseFrameAttribute>& OutDirectionFrameAttributes,
@@ -527,13 +530,14 @@ public:
 		const FAnimDatabaseFrameRanges& FrameRanges,
 		const int32 SampleNum = 4,
 		const float FutureTime = 1.0f,
+		const float PastTime = 0.0f,
 		const FVector ForwardVector = FVector(0.0f, 1.0f, 0.0f));
 
 	/**
 	 * Makes look-at trajectory frame attributes transformed local to the character's root location and rotation. Computes the overall scale of 
 	 * locations - which is useful for scaling distance costs for use with matching.
 	 */
-	UFUNCTION(BlueprintPure = true, Category = "AnimDatabase", meta = (NotBlueprintThreadSafe, AdvancedDisplay = 10))
+	UFUNCTION(BlueprintPure = true, Category = "AnimDatabase", meta = (NotBlueprintThreadSafe, AdvancedDisplay = 11))
 	static UE_API void MakeRootLocalLookAtTrajectoryFrameAttributeAndScale(
 		TArray<FAnimDatabaseFrameAttribute>& OutLocationFrameAttributes,
 		TArray<FAnimDatabaseFrameAttribute>& OutDirectionFrameAttributes,
@@ -543,6 +547,7 @@ public:
 		const int32 LookAtBoneIndex,
 		const int32 SampleNum = 4,
 		const float FutureTime = 1.0f,
+		const float PastTime = 0.0f,
 		const FVector LocalDirection = FVector(0.0f, 1.0f, 0.0f),
 		const float DirectionSmoothingAmount = 15.0f);
 
diff --git a/Engine/Plugins/Experimental/Animation/AnimGen/Source/AnimGen/Private/AnimGenBehavior.cpp b/Engine/Plugins/Experimental/Animation/AnimGen/Source/AnimGen/Private/AnimGenBehavior.cpp
--- a/Engine/Plugins/Experimental/Animation/AnimGen/Source/AnimGen/Private/AnimGenBehavior.cpp
+++ b/Engine/Plugins/Experimental/Animation/AnimGen/Source/AnimGen/Private/AnimGenBehavior.cpp
@@ -266,11 +266,12 @@ FAnimGenControlSchemaElement UAnimGenBehavior_TrajectoryFollow::SpecifyControl_I
 
 void UAnimGenBehavior_TrajectoryFollow::MakeControlSets_Implementation(TArray<FAnimGenControlSet>& OutControlSets, FAnimGenControlObject& InControlObject, const UAnimDatabase* InDatabase, const FAnimDatabaseFrameRanges& InFrameRanges) const
 {
-	// Trim TrajectoryFutureTime from all anim sequences since we don't have the future trajectory
+	// Trim TrajectoryFutureTime and TrajectoryPastTime from all anim sequences since we don't have those trajectories
 	const FAnimDatabaseFrameRanges TrimmedRanges = UAnimDatabaseFrameRangesLibrary::FrameRangesIntersection(
 		InFrameRanges, 
-		UAnimDatabaseFrameRangesLibrary::FrameRangesTrimEnd(
+		UAnimDatabaseFrameRangesLibrary::FrameRangesTrim(
 			UAnimDatabaseFrameRangesLibrary::FrameRangesAll(InDatabase),
+			FMath::RoundToInt(TrajectoryPastTime * InDatabase->GetFrameRate().AsDecimal()),
 			FMath::RoundToInt(TrajectoryFutureTime * InDatabase->GetFrameRate().AsDecimal())));
 
 	const int32 TagNum = TagRanges.Num();
@@ -299,7 +300,7 @@ void UAnimGenBehavior_TrajectoryFollow::MakeControlSets_Implementation(TArray<FA
 			TEXT("Tags")
 		},
 		{
-			UAnimGenControls::MakeTrajectoryControlSetFromDatabase(InControlObject, InDatabase, TrimmedRanges, TrajectorySampleNum, TrajectoryFutureTime, TrajectoryForwardVector),
+			UAnimGenControls::MakeTrajectoryControlSetFromDatabase(InControlObject, InDatabase, TrimmedRanges, TrajectorySampleNum, TrajectoryFutureTime, TrajectoryPastTime, TrajectoryForwardVector),
 			UAnimGenControls::MakeNamedInclusiveDiscreteControlSetFromFrameRangesArrayViews(InControlObject, TrimmedRanges, TagNames, TagFrameRanges),
 		}));
 }
@@ -560,11 +561,12 @@ FAnimGenControlSchemaElement UAnimGenBehavior_TrajectoryInteraction::SpecifyCont
 
 void UAnimGenBehavior_TrajectoryInteraction::MakeControlSets_Implementation(TArray<FAnimGenControlSet>& OutControlSets, FAnimGenControlObject& InControlObject, const UAnimDatabase* InDatabase, const FAnimDatabaseFrameRanges& InFrameRanges) const
 {
-	// Trim TrajectoryFutureTime from all anim sequences since we don't have the future trajectory
+	// Trim TrajectoryFutureTime and TrajectoryPastTime from all anim sequences since we don't have those trajectories
 	const FAnimDatabaseFrameRanges TrimmedRanges = UAnimDatabaseFrameRangesLibrary::FrameRangesIntersection(
 		InFrameRanges,
-		UAnimDatabaseFrameRangesLibrary::FrameRangesTrimEnd(
+		UAnimDatabaseFrameRangesLibrary::FrameRangesTrim(
 			UAnimDatabaseFrameRangesLibrary::FrameRangesAll(InDatabase),
+			FMath::RoundToInt(TrajectoryPastTime * InDatabase->GetFrameRate().AsDecimal()),
 			FMath::RoundToInt(TrajectoryFutureTime * InDatabase->GetFrameRate().AsDecimal())));
 
 	const int32 TagNum = TagRanges.Num();
@@ -612,7 +614,7 @@ void UAnimGenBehavior_TrajectoryInteraction::MakeControlSets_Implementation(TArr
 			TEXT("Interaction")
 		},
 		{
-			UAnimGenControls::MakeTrajectoryControlSetFromDatabase(InControlObject, InDatabase, TrimmedRanges, TrajectorySampleNum, TrajectoryFutureTime, TrajectoryForwardVector),
+			UAnimGenControls::MakeTrajectoryControlSetFromDatabase(InControlObject, InDatabase, TrimmedRanges, TrajectorySampleNum, TrajectoryFutureTime, TrajectoryPastTime, TrajectoryForwardVector),
 			UAnimGenControls::MakeNamedInclusiveDiscreteControlSetFromFrameRangesArrayViews(InControlObject, TrimmedRanges, TagNames, TagFrameRanges),
 			UAnimGenControls::MakeOptionalControlSet(InControlObject, UAnimGenControls::MakeStructControlSetFromArrayViews(InControlObject, 
 					{
@@ -840,6 +842,7 @@ FAnimGenControlObjectElement UAnimGenBehaviorLibrary::MakeTrajectoryFollowBehavi
 	const TArrayView<const FName> DesiredTags,
 	const int32 InTrajectorySampleNum,
 	const float InTrajectoryFutureTime,
+	const float InTrajectoryPastTime,
 	const FVector InTrajectoryForwardVector)
 {
 	return UAnimGenControls::MakeStructControlFromArrayViews(InControlObject,
@@ -848,7 +851,7 @@ FAnimGenControlObjectElement UAnimGenBehaviorLibrary::MakeTrajectoryFollowBehavi
 			TEXT("Tags")
 		},
 		{
-			UAnimGenControls::MakeTrajectoryControlFromTransformTrajectory(InControlObject, TransformTrajectory, RootBoneTransform, InTrajectorySampleNum, InTrajectoryFutureTime, InTrajectoryForwardVector),
+			UAnimGenControls::MakeTrajectoryControlFromTransformTrajectory(InControlObject, TransformTrajectory, RootBoneTransform, InTrajectorySampleNum, InTrajectoryFutureTime, InTrajectoryPastTime, InTrajectoryForwardVector),
 			UAnimGenControls::MakeNamedInclusiveDiscreteControlFromArrayView(InControlObject, DesiredTags),
 		});
 }
@@ -860,9 +863,10 @@ FAnimGenControlObjectElement UAnimGenBehaviorLibrary::MakeTrajectoryFollowBehavi
 	const TArray<FName>& DesiredTags,
 	const int32 InTrajectorySampleNum,
 	const float InTrajectoryFutureTime,
+	const float InTrajectoryPastTime,
 	const FVector InTrajectoryForwardVector)
 {
-	return MakeTrajectoryFollowBehaviorControlFromArrayView(InControlObject, TransformTrajectory, RootBoneTransform, DesiredTags, InTrajectorySampleNum, InTrajectoryFutureTime, InTrajectoryForwardVector);
+	return MakeTrajectoryFollowBehaviorControlFromArrayView(InControlObject, TransformTrajectory, RootBoneTransform, DesiredTags, InTrajectorySampleNum, InTrajectoryFutureTime, InTrajectoryPastTime, InTrajectoryForwardVector);
 }
 
 FAnimGenControlObjectElement UAnimGenBehaviorLibrary::MakeMoveToTargetBehaviorControl(
@@ -909,9 +913,10 @@ FAnimGenControlObjectElement UAnimGenBehaviorLibrary::MakeTrajectoryInteractionB
 	const bool bInteractionActive,
 	const int32 InTrajectorySampleNum,
 	const float InTrajectoryFutureTime,
+	const float InTrajectoryPastTime,
 	const FVector InTrajectoryForwardVector)
 {
-	return MakeTrajectoryInteractionBehaviorControlFromArrayView(InControlObject, TransformTrajectory, RootBoneTransform, DesiredTags, InteractionLocation, InteractionDirection, bInteractionActive, InTrajectorySampleNum, InTrajectoryFutureTime, InTrajectoryForwardVector);
+	return MakeTrajectoryInteractionBehaviorControlFromArrayView(InControlObject, TransformTrajectory, RootBoneTransform, DesiredTags, InteractionLocation, InteractionDirection, bInteractionActive, InTrajectorySampleNum, InTrajectoryFutureTime, InTrajectoryPastTime, InTrajectoryForwardVector);
 }
 
 FAnimGenControlObjectElement UAnimGenBehaviorLibrary::MakeTrajectoryInteractionBehaviorControlFromArrayView(
@@ -924,6 +929,7 @@ FAnimGenControlObjectElement UAnimGenBehaviorLibrary::MakeTrajectoryInteractionB
 	const bool bInteractionActive,
 	const int32 InTrajectorySampleNum,
 	const float InTrajectoryFutureTime,
+	const float InTrajectoryPastTime,
 	const FVector InTrajectoryForwardVector)
 {
 	return UAnimGenControls::MakeStructControlFromArrayViews(InControlObject,
@@ -933,7 +939,7 @@ FAnimGenControlObjectElement UAnimGenBehaviorLibrary::MakeTrajectoryInteractionB
 			TEXT("Interaction")
 		},
 		{
-			UAnimGenControls::MakeTrajectoryControlFromTransformTrajectory(InControlObject, TransformTrajectory, RootBoneTransform, InTrajectorySampleNum, InTrajectoryFutureTime, InTrajectoryForwardVector),
+			UAnimGenControls::MakeTrajectoryControlFromTransformTrajectory(InControlObject, TransformTrajectory, RootBoneTransform, InTrajectorySampleNum, InTrajectoryFutureTime, InTrajectoryPastTime, InTrajectoryForwardVector),
 			UAnimGenControls::MakeNamedInclusiveDiscreteControlFromArrayView(InControlObject, DesiredTags),
 			bInteractionActive ? 
 				UAnimGenControls::MakeOptionalValidControl(InControlObject, UAnimGenControls::MakeStructControlFromArrayViews(InControlObject,
@@ -1400,26 +1406,26 @@ void UAnimGenBehaviorLibrary::MakeTransformTrajectoryFromArrayViews(FTransformTr
 	}
 }
 
-void UAnimGenBehaviorLibrary::MakeTransformTrajectoryFromPath(FTransformTrajectory& OutTransformTrajectory, const TArray<FVector>& Path, const float StartDistance, const float PathPointDistance, const float PathSpeed, const float FutureTime, const float SampleInterval)
+void UAnimGenBehaviorLibrary::MakeTransformTrajectoryFromPath(FTransformTrajectory& OutTransformTrajectory, const TArray<FVector>& Path, const float StartDistance, const float PathPointDistance, const float PathSpeed, const float FutureTime, const float PastTime, const float SampleInterval)
 {
-	return MakeTransformTrajectoryFromPathArrayView(OutTransformTrajectory, Path, StartDistance, PathPointDistance, PathSpeed, FutureTime, SampleInterval);
+	return MakeTransformTrajectoryFromPathArrayView(OutTransformTrajectory, Path, StartDistance, PathPointDistance, PathSpeed, FutureTime, PastTime, SampleInterval);
 }
 
-void UAnimGenBehaviorLibrary::MakeTransformTrajectoryFromPathArrayView(FTransformTrajectory& OutTransformTrajectory, const TArrayView<const FVector> Path, const float StartDistance, const float PathPointDistance, const float PathSpeed, const float FutureTime, const float SampleInterval)
+void UAnimGenBehaviorLibrary::MakeTransformTrajectoryFromPathArrayView(FTransformTrajectory& OutTransformTrajectory, const TArrayView<const FVector> Path, const float StartDistance, const float PathPointDistance, const float PathSpeed, const float FutureTime, const float PastTime, const float SampleInterval)
 {
-	const int32 SampleNum = FMath::Max(FMath::RoundToInt(FutureTime / FMath::Max(SampleInterval, UE_SMALL_NUMBER)), 0);
+	const int32 SampleNum = FMath::Max(FMath::RoundToInt((FutureTime + PastTime) / FMath::Max(SampleInterval, UE_SMALL_NUMBER)), 0);
 	OutTransformTrajectory.Samples.SetNumUninitialized(SampleNum);
 	for (int32 SampleIdx = 0; SampleIdx < SampleNum; SampleIdx++)
 	{
-		const float SamplePathTime = StartDistance + PathSpeed * SampleIdx * SampleInterval;
+		const float SamplePathTime = StartDistance + PathSpeed * (SampleIdx * SampleInterval - PastTime);
 
 		OutTransformTrajectory.Samples[SampleIdx].Position = PathLocationAtTimeFromArrayView(Path, SamplePathTime / FMath::Max(PathPointDistance, UE_SMALL_NUMBER));
 		OutTransformTrajectory.Samples[SampleIdx].Facing = PathRotationAtTimeFromArrayView(Path, SamplePathTime / FMath::Max(PathPointDistance, UE_SMALL_NUMBER)).Quaternion();
-		OutTransformTrajectory.Samples[SampleIdx].TimeInSeconds = SampleIdx * SampleInterval;
+		OutTransformTrajectory.Samples[SampleIdx].TimeInSeconds = SampleIdx * SampleInterval - PastTime;
 	}
 }
 
-void UAnimGenBehaviorLibrary::MakeTransformTrajectoryFromTrajectory(FTransformTrajectory& OutTransformTrajectory, const FAnimGenTrajectory& Trajectory, const float FutureTime)
+void UAnimGenBehaviorLibrary::MakeTransformTrajectoryFromTrajectory(FTransformTrajectory& OutTransformTrajectory, const FAnimGenTrajectory& Trajectory, const float FutureTime, const float PastTime)
 {
 	const int32 SampleNum = FMath::Min(Trajectory.Locations.Num(), Trajectory.Rotations.Num());
 
@@ -1428,7 +1434,7 @@ void UAnimGenBehaviorLibrary::MakeTransformTrajectoryFromTrajectory(FTransformTr
 	{
 		OutTransformTrajectory.Samples[SampleIdx].Position = Trajectory.Locations[SampleIdx];
 		OutTransformTrajectory.Samples[SampleIdx].Facing = Trajectory.Rotations[SampleIdx].Quaternion();
-		OutTransformTrajectory.Samples[SampleIdx].TimeInSeconds = FutureTime * (((float)SampleIdx) / FMath::Max(SampleNum - 1, 1));
+		OutTransformTrajectory.Samples[SampleIdx].TimeInSeconds = (FutureTime + PastTime) * (((float)SampleIdx) / FMath::Max(SampleNum - 1, 1)) - PastTime;
 	}
 }
 
@@ -1842,6 +1848,7 @@ FAnimGenControlObjectElement UAnimGenBehaviorPreview_TrajectoryFollow::MakeContr
 		DesiredTags,
 		TrajectorySampleNum,
 		TrajectoryFutureTime,
+		TrajectoryPastTime,
 		TrajectoryForwardVector);
 }
 
diff --git a/Engine/Plugins/Experimental/Animation/AnimGen/Source/AnimGen/Private/AnimGenControl.cpp b/Engine/Plugins/Experimental/Animation/AnimGen/Source/AnimGen/Private/AnimGenControl.cpp
--- a/Engine/Plugins/Experimental/Animation/AnimGen/Source/AnimGen/Private/AnimGenControl.cpp
+++ b/Engine/Plugins/Experimental/Animation/AnimGen/Source/AnimGen/Private/AnimGenControl.cpp
@@ -2342,7 +2342,7 @@ FAnimGenControlObjectElement UAnimGenControls::MakeLocationTrajectoryControlFrom
 	return MakeLocationsStaticArrayControlFromArrayView(Object, Locations, RelativeTransform, Tag);
 }
 
-FAnimGenControlObjectElement UAnimGenControls::MakeTrajectoryControlFromTransformTrajectory(FAnimGenControlObject& Object, const FTransformTrajectory& TransformTrajectory, const FTransform& RelativeTransform, const int32 SampleNum, const float FutureTime, const FVector ForwardVector, const FName Tag)
+FAnimGenControlObjectElement UAnimGenControls::MakeTrajectoryControlFromTransformTrajectory(FAnimGenControlObject& Object, const FTransformTrajectory& TransformTrajectory, const FTransform& RelativeTransform, const int32 SampleNum, const float FutureTime, const float PastTime, const FVector ForwardVector, const FName Tag)
 {
 	if (SampleNum < 0)
 	{
@@ -2357,7 +2357,7 @@ FAnimGenControlObjectElement UAnimGenControls::MakeTrajectoryControlFromTransfor
 
 	for (int32 SampleIdx = 0; SampleIdx < SampleNum; SampleIdx++)
 	{
-		const float Time = FutureTime * ((float)(SampleIdx) / FMath::Max(SampleNum - 1, 1));
+		const float Time = (FutureTime + PastTime) * ((float)(SampleIdx) / FMath::Max(SampleNum - 1, 1)) - PastTime;
 		FTransformTrajectorySample Sample = TransformTrajectory.GetSampleAtTime(Time);
 		Locations[SampleIdx] = Sample.Position;
 		Directions[SampleIdx] = Sample.Facing.RotateVector(ForwardVector);
@@ -2366,7 +2366,7 @@ FAnimGenControlObjectElement UAnimGenControls::MakeTrajectoryControlFromTransfor
 	return MakeTrajectoryControlFromArrayViews(Object, Locations, Directions, RelativeTransform, Tag);
 }
 
-FAnimGenControlObjectElement UAnimGenControls::MakeLocationTrajectoryControlFromTransformTrajectory(FAnimGenControlObject& Object, const FTransformTrajectory& TransformTrajectory, const FTransform& RelativeTransform, const int32 SampleNum, const float FutureTime, const FName Tag)
+FAnimGenControlObjectElement UAnimGenControls::MakeLocationTrajectoryControlFromTransformTrajectory(FAnimGenControlObject& Object, const FTransformTrajectory& TransformTrajectory, const FTransform& RelativeTransform, const int32 SampleNum, const float FutureTime, const float PastTime, const FName Tag)
 {
 	if (SampleNum < 0)
 	{
@@ -2379,7 +2379,7 @@ FAnimGenControlObjectElement UAnimGenControls::MakeLocationTrajectoryControlFrom
 
 	for (int32 SampleIdx = 0; SampleIdx < SampleNum; SampleIdx++)
 	{
-		const float Time = FutureTime * ((float)(SampleIdx) / FMath::Max(SampleNum - 1, 1));
+		const float Time = (FutureTime + PastTime) * ((float)(SampleIdx) / FMath::Max(SampleNum - 1, 1)) - PastTime;
 		FTransformTrajectorySample Sample = TransformTrajectory.GetSampleAtTime(Time);
 		Locations[SampleIdx] = Sample.Position;
 	}
@@ -5824,11 +5824,11 @@ FAnimGenControlSet UAnimGenControls::MakeRootDirectionAtRangeEndControlSetFromDa
 		Tag);
 }
 
-FAnimGenControlSet UAnimGenControls::MakeTrajectoryControlSetFromDatabase(FAnimGenControlObject& Object, const UAnimDatabase* Database, const FAnimDatabaseFrameRanges& FrameRanges, const int32 SampleNum, const float FutureTime, const FVector ForwardVector, const FName Tag)
+FAnimGenControlSet UAnimGenControls::MakeTrajectoryControlSetFromDatabase(FAnimGenControlObject& Object, const UAnimDatabase* Database, const FAnimDatabaseFrameRanges& FrameRanges, const int32 SampleNum, const float FutureTime, const float PastTime, const FVector ForwardVector, const FName Tag)
 {
 	TArray<FAnimDatabaseFrameAttribute> TrajectoryLocations;
 	TArray<FAnimDatabaseFrameAttribute> TrajectoryDirections;
-	UAnimDatabaseFrameAttributeLibrary::MakeTrajectoryFrameAttribute(TrajectoryLocations, TrajectoryDirections, Database, FrameRanges, SampleNum, FutureTime, ForwardVector);
+	UAnimDatabaseFrameAttributeLibrary::MakeTrajectoryFrameAttribute(TrajectoryLocations, TrajectoryDirections, Database, FrameRanges, SampleNum, FutureTime, PastTime, ForwardVector);
 
 	return MakeTrajectoryControlSet(Object, 
 		TrajectoryLocations, 
@@ -5837,11 +5837,11 @@ FAnimGenControlSet UAnimGenControls::MakeTrajectoryControlSetFromDatabase(FAnimG
 		Tag);
 }
 
-FAnimGenControlSet UAnimGenControls::MakeLookAtTrajectoryControlSetFromDatabase(FAnimGenControlObject& Object, const UAnimDatabase* Database, const FAnimDatabaseFrameRanges& FrameRanges, const int32 LookAtBoneIndex, const int32 SampleNum, const float FutureTime, const FVector LocalDirection, const float DirectionSmoothingAmount, const FName Tag)
+FAnimGenControlSet UAnimGenControls::MakeLookAtTrajectoryControlSetFromDatabase(FAnimGenControlObject& Object, const UAnimDatabase* Database, const FAnimDatabaseFrameRanges& FrameRanges, const int32 LookAtBoneIndex, const int32 SampleNum, const float FutureTime, const float PastTime, const FVector LocalDirection, const float DirectionSmoothingAmount, const FName Tag)
 {
 	TArray<FAnimDatabaseFrameAttribute> TrajectoryLocations;
 	TArray<FAnimDatabaseFrameAttribute> TrajectoryDirections;
-	UAnimDatabaseFrameAttributeLibrary::MakeLookAtTrajectoryFrameAttribute(TrajectoryLocations, TrajectoryDirections, Database, FrameRanges, LookAtBoneIndex, SampleNum, FutureTime, LocalDirection, DirectionSmoothingAmount);
+	UAnimDatabaseFrameAttributeLibrary::MakeLookAtTrajectoryFrameAttribute(TrajectoryLocations, TrajectoryDirections, Database, FrameRanges, LookAtBoneIndex, SampleNum, FutureTime, PastTime, LocalDirection, DirectionSmoothingAmount);
 
 	return MakeTrajectoryControlSet(Object,
 		TrajectoryLocations,
@@ -5850,10 +5850,10 @@ FAnimGenControlSet UAnimGenControls::MakeLookAtTrajectoryControlSetFromDatabase(
 		Tag);
 }
 
-FAnimGenControlSet UAnimGenControls::MakeLocationTrajectoryControlSetFromDatabase(FAnimGenControlObject& Object, const UAnimDatabase* Database, const FAnimDatabaseFrameRanges& FrameRanges, const int32 SampleNum, const float FutureTime, const FName Tag)
+FAnimGenControlSet UAnimGenControls::MakeLocationTrajectoryControlSetFromDatabase(FAnimGenControlObject& Object, const UAnimDatabase* Database, const FAnimDatabaseFrameRanges& FrameRanges, const int32 SampleNum, const float FutureTime, const float PastTime, const FName Tag)
 {
 	TArray<FAnimDatabaseFrameAttribute> TrajectoryLocations;
-	UAnimDatabaseFrameAttributeLibrary::MakeLocationTrajectoryFrameAttribute(TrajectoryLocations, Database, FrameRanges, SampleNum, FutureTime);
+	UAnimDatabaseFrameAttributeLibrary::MakeLocationTrajectoryFrameAttribute(TrajectoryLocations, Database, FrameRanges, SampleNum, FutureTime, PastTime);
 
 	return MakeLocationTrajectoryControlSet(Object,
 		TrajectoryLocations,
diff --git a/Engine/Plugins/Experimental/Animation/AnimGen/Source/AnimGen/Private/RigVMFunction_AnimGen.cpp b/Engine/Plugins/Experimental/Animation/AnimGen/Source/AnimGen/Private/RigVMFunction_AnimGen.cpp
--- a/Engine/Plugins/Experimental/Animation/AnimGen/Source/AnimGen/Private/RigVMFunction_AnimGen.cpp
+++ b/Engine/Plugins/Experimental/Animation/AnimGen/Source/AnimGen/Private/RigVMFunction_AnimGen.cpp
@@ -11,7 +11,7 @@ FRigVMFunction_MakeControlObject_Execute()
 
 FRigVMFunction_MakeTrajectoryControl_Execute()
 {
-	Result = UAnimGenControls::MakeTrajectoryControlFromTransformTrajectory(Object, TransformTrajectory, RelativeTransform, SampleNum, FutureTime, ForwardVector, Tag);
+	Result = UAnimGenControls::MakeTrajectoryControlFromTransformTrajectory(Object, TransformTrajectory, RelativeTransform, SampleNum, FutureTime, PastTime, ForwardVector, Tag);
 }
 
 FRigVMFunction_MakeStructControl_Execute()
@@ -45,5 +45,6 @@ FRigVMFunction_MakeTrajectoryFollowBehaviorControl_Execute()
 		DesiredTags,
 		TrajectorySampleNum,
 		TrajectoryFutureTime,
+		TrajectoryPastTime,
 		TrajectoryForwardVector);
 }
diff --git a/Engine/Plugins/Experimental/Animation/AnimGen/Source/AnimGen/Public/AnimGenBehavior.h b/Engine/Plugins/Experimental/Animation/AnimGen/Source/AnimGen/Public/AnimGenBehavior.h
--- a/Engine/Plugins/Experimental/Animation/AnimGen/Source/AnimGen/Public/AnimGenBehavior.h
+++ b/Engine/Plugins/Experimental/Animation/AnimGen/Source/AnimGen/Public/AnimGenBehavior.h
@@ -333,6 +333,10 @@ public:
 	UPROPERTY(EditAnywhere, Category = "Trajectory")
 	float TrajectoryFutureTime = 1.0f;
 
+	/** Past time for the trajectory */
+	UPROPERTY(EditAnywhere, Category = "Trajectory")
+	float TrajectoryPastTime = 0.0f;
+
 	/** Number of samples to use for the trajectory */
 	UPROPERTY(EditAnywhere, Category = "Trajectory")
 	int32 TrajectorySampleNum = 4;
@@ -468,6 +472,10 @@ public:
 	UPROPERTY(EditAnywhere, Category = "Trajectory")
 	float TrajectoryFutureTime = 1.0f;
 
+	/** Past time for the trajectory */
+	UPROPERTY(EditAnywhere, Category = "Trajectory")
+	float TrajectoryPastTime = 0.0f;
+
 	/** Number of samples to use for the trajectory */
 	UPROPERTY(EditAnywhere, Category = "Trajectory")
 	int32 TrajectorySampleNum = 4;
@@ -645,6 +653,7 @@ public:
 		const TArray<FName>& DesiredTags,
 		const int32 InTrajectorySampleNum = 4,
 		const float InTrajectoryFutureTime = 1.0f,
+		const float InTrajectoryPastTime = 0.0f,
 		const FVector InTrajectoryForwardVector = FVector::RightVector);
 
 	static UE_API FAnimGenControlObjectElement MakeTrajectoryFollowBehaviorControlFromArrayView(
@@ -654,6 +663,7 @@ public:
 		const TArrayView<const FName> DesiredTags = {},
 		const int32 InTrajectorySampleNum = 4,
 		const float InTrajectoryFutureTime = 1.0f,
+		const float InTrajectoryPastTime = 0.0f,
 		const FVector InTrajectoryForwardVector = FVector::RightVector);
 
 	/** Make a control object element for the Move To Target Behavior */
@@ -686,6 +696,7 @@ public:
 		const bool bInteractionActive = false,
 		const int32 InTrajectorySampleNum = 4,
 		const float InTrajectoryFutureTime = 1.0f,
+		const float InTrajectoryPastTime = 0.0f,
 		const FVector InTrajectoryForwardVector = FVector::RightVector);
 
 	static UE_API FAnimGenControlObjectElement MakeTrajectoryInteractionBehaviorControlFromArrayView(
@@ -698,6 +709,7 @@ public:
 		const bool bInteractionActive = false,
 		const int32 InTrajectorySampleNum = 4,
 		const float InTrajectoryFutureTime = 1.0f,
+		const float InTrajectoryPastTime = 0.0f,
 		const FVector InTrajectoryForwardVector = FVector::RightVector);
 
 	/** Make a control object element for the Multi Behavior */
@@ -936,12 +948,12 @@ public:
 
 	/** Make a transform trajectory from an AnimGenTrajectory */
 	UFUNCTION(BlueprintPure = true, Category = "AnimGen")
-	static UE_API void MakeTransformTrajectoryFromTrajectory(FTransformTrajectory& OutTransformTrajectory, const FAnimGenTrajectory& Trajectory, const float FutureTime = 1.0f);
+	static UE_API void MakeTransformTrajectoryFromTrajectory(FTransformTrajectory& OutTransformTrajectory, const FAnimGenTrajectory& Trajectory, const float FutureTime = 1.0f, const float PastTime = 0.0f);
 
 	/** Make a transform trajectory from a path */
 	UFUNCTION(BlueprintPure = true, Category = "AnimGen")
-	static UE_API void MakeTransformTrajectoryFromPath(FTransformTrajectory& OutTransformTrajectory, const TArray<FVector>& Path, const float StartDistance = 0.0f, const float PathPointDistance = 50.0f, const float PathSpeed = 150.0f, const float FutureTime = 1.0f, const float SampleInterval = 0.01666666666f);
-	static UE_API void MakeTransformTrajectoryFromPathArrayView(FTransformTrajectory& OutTransformTrajectory, const TArrayView<const FVector> Path, const float StartDistance = 0.0f, const float PathPointDistance = 50.0f, const float PathSpeed = 150.0f, const float FutureTime = 1.0f, const float SampleInterval = 0.01666666666f);
+	static UE_API void MakeTransformTrajectoryFromPath(FTransformTrajectory& OutTransformTrajectory, const TArray<FVector>& Path, const float StartDistance = 0.0f, const float PathPointDistance = 50.0f, const float PathSpeed = 150.0f, const float FutureTime = 1.0f, const float PastTime = 0.0f, const float SampleInterval = 0.01666666666f);
+	static UE_API void MakeTransformTrajectoryFromPathArrayView(FTransformTrajectory& OutTransformTrajectory, const TArrayView<const FVector> Path, const float StartDistance = 0.0f, const float PathPointDistance = 50.0f, const float PathSpeed = 150.0f, const float FutureTime = 1.0f, const float PastTime = 0.0f, const float SampleInterval = 0.01666666666f);
 
 	/** Project a transform trajectory onto the ground */
 	UFUNCTION(BlueprintPure = false, Category = "AnimGen")
@@ -1146,6 +1158,10 @@ public:
 	UPROPERTY(EditAnywhere, Category = "Trajectory", meta = (ForceUnits = "s"))
 	float TrajectoryFutureTime = 1.0f;
 
+	/** Past time for the trajectory */
+	UPROPERTY(EditAnywhere, Category = "Trajectory", meta = (ForceUnits = "s"))
+	float TrajectoryPastTime = 0.0f;
+
 	/** Number of samples to use for the trajectory */
 	UPROPERTY(EditAnywhere, Category = "Trajectory")
 	int32 TrajectorySampleNum = 4;
diff --git a/Engine/Plugins/Experimental/Animation/AnimGen/Source/AnimGen/Public/AnimGenControl.h b/Engine/Plugins/Experimental/Animation/AnimGen/Source/AnimGen/Public/AnimGenControl.h
--- a/Engine/Plugins/Experimental/Animation/AnimGen/Source/AnimGen/Public/AnimGenControl.h
+++ b/Engine/Plugins/Experimental/Animation/AnimGen/Source/AnimGen/Public/AnimGenControl.h
@@ -536,10 +536,10 @@ public:
 	static UE_API FAnimGenControlObjectElement MakeLocationTrajectoryControlFromArrayViews(UPARAM(ref) FAnimGenControlObject& Object, const TArrayView<const FVector> Locations, const FTransform& RelativeTransform = FTransform(), const FName Tag = TEXT("LocationTrajectoryControl"));
 
 	UFUNCTION(BlueprintPure, Category = "AnimGen", meta = (AdvancedDisplay = 6, AutoCreateRefTerm = "RelativeTransform"))
-	static UE_API FAnimGenControlObjectElement MakeTrajectoryControlFromTransformTrajectory(UPARAM(ref) FAnimGenControlObject& Object, const FTransformTrajectory& TransformTrajectory, const FTransform& RelativeTransform = FTransform(), const int32 SampleNum = 4, const float FutureTime = 1.0f, const FVector ForwardVector = FVector(0,1,0), const FName Tag = TEXT("TrajectoryControl"));
+	static UE_API FAnimGenControlObjectElement MakeTrajectoryControlFromTransformTrajectory(UPARAM(ref) FAnimGenControlObject& Object, const FTransformTrajectory& TransformTrajectory, const FTransform& RelativeTransform = FTransform(), const int32 SampleNum = 4, const float FutureTime = 1.0f, const float PastTime = 0.0f, const FVector ForwardVector = FVector(0,1,0), const FName Tag = TEXT("TrajectoryControl"));
 
 	UFUNCTION(BlueprintPure, Category = "AnimGen", meta = (AdvancedDisplay = 5, AutoCreateRefTerm = "RelativeTransform"))
-	static UE_API FAnimGenControlObjectElement MakeLocationTrajectoryControlFromTransformTrajectory(UPARAM(ref) FAnimGenControlObject& Object, const FTransformTrajectory& TransformTrajectory, const FTransform& RelativeTransform = FTransform(), const int32 SampleNum = 4, const float FutureTime = 1.0f, const FName Tag = TEXT("LocationTrajectoryControl"));
+	static UE_API FAnimGenControlObjectElement MakeLocationTrajectoryControlFromTransformTrajectory(UPARAM(ref) FAnimGenControlObject& Object, const FTransformTrajectory& TransformTrajectory, const FTransform& RelativeTransform = FTransform(), const int32 SampleNum = 4, const float FutureTime = 1.0f, const float PastTime = 0.0f, const FName Tag = TEXT("LocationTrajectoryControl"));
 
 	UFUNCTION(BlueprintPure, Category = "AnimGen", meta = (AdvancedDisplay = 3))
 	static UE_API FAnimGenControlObjectElement MakeEventControl(UPARAM(ref) FAnimGenControlObject& Object, const bool bTimeUntilEventKnown, const float TimeUntilEvent, const FName Tag = TEXT("EventControl"));
@@ -901,14 +901,14 @@ public:
 	UFUNCTION(BlueprintPure, Category = "AnimGen", meta = (AdvancedDisplay = 4))
 	static UE_API FAnimGenControlSet MakeRootDirectionAtRangeEndControlSetFromDatabase(UPARAM(ref) FAnimGenControlObject& Object, const UAnimDatabase* Database, const FAnimDatabaseFrameRanges& FrameRanges, const FVector ForwardVector = FVector(0.0, 1.0f, 0.0f), const FName Tag = TEXT("DirectionControl"));
 
-	UFUNCTION(BlueprintPure, Category = "AnimGen", meta = (AdvancedDisplay = 6))
-	static UE_API FAnimGenControlSet MakeTrajectoryControlSetFromDatabase(UPARAM(ref) FAnimGenControlObject& Object, const UAnimDatabase* Database, const FAnimDatabaseFrameRanges& FrameRanges, const int32 SampleNum = 4, const float FutureTime = 1.0f, const FVector ForwardVector = FVector(0.0, 1.0f, 0.0f), const FName Tag = TEXT("TrajectoryControl"));
+	UFUNCTION(BlueprintPure, Category = "AnimGen", meta = (AdvancedDisplay = 7))
+	static UE_API FAnimGenControlSet MakeTrajectoryControlSetFromDatabase(UPARAM(ref) FAnimGenControlObject& Object, const UAnimDatabase* Database, const FAnimDatabaseFrameRanges& FrameRanges, const int32 SampleNum = 4, const float FutureTime = 1.0f, const float PastTime = 0.0f, const FVector ForwardVector = FVector(0.0, 1.0f, 0.0f), const FName Tag = TEXT("TrajectoryControl"));
 
-	UFUNCTION(BlueprintPure, Category = "AnimGen", meta = (AdvancedDisplay = 8))
-	static UE_API FAnimGenControlSet MakeLookAtTrajectoryControlSetFromDatabase(UPARAM(ref) FAnimGenControlObject& Object, const UAnimDatabase* Database, const FAnimDatabaseFrameRanges& FrameRanges, const int32 LookAtBoneIndex, const int32 SampleNum = 4, const float FutureTime = 1.0f, const FVector LocalDirection = FVector(0.0, 1.0f, 0.0f), const float DirectionSmoothingAmount = 15.0f, const FName Tag = TEXT("TrajectoryControl"));
+	UFUNCTION(BlueprintPure, Category = "AnimGen", meta = (AdvancedDisplay = 9))
+	static UE_API FAnimGenControlSet MakeLookAtTrajectoryControlSetFromDatabase(UPARAM(ref) FAnimGenControlObject& Object, const UAnimDatabase* Database, const FAnimDatabaseFrameRanges& FrameRanges, const int32 LookAtBoneIndex, const int32 SampleNum = 4, const float FutureTime = 1.0f, const float PastTime = 0.0f, const FVector LocalDirection = FVector(0.0, 1.0f, 0.0f), const float DirectionSmoothingAmount = 15.0f, const FName Tag = TEXT("TrajectoryControl"));
 
 	UFUNCTION(BlueprintPure, Category = "AnimGen", meta = (AdvancedDisplay = 5))
-	static UE_API FAnimGenControlSet MakeLocationTrajectoryControlSetFromDatabase(UPARAM(ref) FAnimGenControlObject& Object, const UAnimDatabase* Database, const FAnimDatabaseFrameRanges& FrameRanges, const int32 SampleNum = 4, const float FutureTime = 1.0f, const FName Tag = TEXT("LocationTrajectoryControl"));
+	static UE_API FAnimGenControlSet MakeLocationTrajectoryControlSetFromDatabase(UPARAM(ref) FAnimGenControlObject& Object, const UAnimDatabase* Database, const FAnimDatabaseFrameRanges& FrameRanges, const int32 SampleNum = 4, const float FutureTime = 1.0f, const float PastTime = 0.0f, const FName Tag = TEXT("LocationTrajectoryControl"));
 
 	UFUNCTION(BlueprintPure, Category = "AnimGen", meta = (AdvancedDisplay = 3))
 	static UE_API FAnimGenControlSet MakeBoolControlSetFromConstant(UPARAM(ref) FAnimGenControlObject& Object, const FAnimDatabaseFrameRanges& FrameRanges, const bool bValue, const FName Tag = TEXT("BoolControl"));
diff --git a/Engine/Plugins/Experimental/Animation/AnimGen/Source/AnimGen/Public/RigVMFunction_AnimGen.h b/Engine/Plugins/Experimental/Animation/AnimGen/Source/AnimGen/Public/RigVMFunction_AnimGen.h
--- a/Engine/Plugins/Experimental/Animation/AnimGen/Source/AnimGen/Public/RigVMFunction_AnimGen.h
+++ b/Engine/Plugins/Experimental/Animation/AnimGen/Source/AnimGen/Public/RigVMFunction_AnimGen.h
@@ -90,6 +90,10 @@ struct FRigVMFunction_MakeTrajectoryControl : public FRigVMFunction_AnimGenBase
 	UPROPERTY(meta = (Input))
 	float FutureTime = 1.0f;
 
+	// Time in the past for the trajectory
+	UPROPERTY(meta = (Input))
+	float PastTime = 0.0f;
+
 	// Local forward vector of the trajectory
 	UPROPERTY(meta = (Input))
 	FVector ForwardVector = FVector(0,1,0);
@@ -188,6 +192,10 @@ struct FRigVMFunction_MakeTrajectoryFollowBehaviorControl : public FRigVMFunctio
 	UPROPERTY(meta = (Input))
 	float TrajectoryFutureTime = 1.0f;
 
+	// Time in the past for the trajectory
+	UPROPERTY(meta = (Input))
+	float TrajectoryPastTime = 0.0f;
+
 	// Local forward vector of the trajectory
 	UPROPERTY(meta = (Input))
 	FVector TrajectoryForwardVector = FVector(0, 1, 0);
diff --git a/Engine/Plugins/Experimental/Animation/AnimGen/Source/AnimGenEditor/Private/AnimGenEditorAutoEncoderToolkit.cpp b/Engine/Plugins/Experimental/Animation/AnimGen/Source/AnimGenEditor/Private/AnimGenEditorAutoEncoderToolkit.cpp
--- a/Engine/Plugins/Experimental/Animation/AnimGen/Source/AnimGenEditor/Private/AnimGenEditorAutoEncoderToolkit.cpp
+++ b/Engine/Plugins/Experimental/Animation/AnimGen/Source/AnimGenEditor/Private/AnimGenEditorAutoEncoderToolkit.cpp
@@ -1332,7 +1332,7 @@ namespace UE::AnimGen::Editor
 
 		TOptional<FTimespan> GetEstimateTimeRemaining() const
 		{
-			if (Iteration > 200 && StartTrainingTimeIteration != -1)
+			if (StartTrainingTimeIteration != -1 && Iteration > StartTrainingTimeIteration + 250)
 			{
 				const FTimespan ElapsedTime = FDateTime::Now() - StartTrainingTime;
 				const int32 ElapsedIteration = Iteration - StartTrainingTimeIteration;
@@ -1358,7 +1358,7 @@ namespace UE::AnimGen::Editor
 
 					// Update Estimated Time
 
-					if (Iteration > 100 && StartTrainingTimeIteration == -1)
+					if (Iteration > 250 && StartTrainingTimeIteration == -1)
 					{
 						StartTrainingTimeIteration = Iteration;
 						StartTrainingTime = FDateTime::Now();
diff --git a/Engine/Plugins/Experimental/Animation/AnimGen/Source/AnimGenEditor/Private/AnimGenEditorControllerToolkit.cpp b/Engine/Plugins/Experimental/Animation/AnimGen/Source/AnimGenEditor/Private/AnimGenEditorControllerToolkit.cpp
--- a/Engine/Plugins/Experimental/Animation/AnimGen/Source/AnimGenEditor/Private/AnimGenEditorControllerToolkit.cpp
+++ b/Engine/Plugins/Experimental/Animation/AnimGen/Source/AnimGenEditor/Private/AnimGenEditorControllerToolkit.cpp
@@ -191,7 +191,8 @@ namespace UE::AnimGen::Editor
 		DataCategory.AddProperty(GET_MEMBER_NAME_CHECKED(UAnimGenControllerTrainingSettings, AutoEncoder), UAnimGenControllerTrainingSettings::StaticClass());
 		DataCategory.AddProperty(GET_MEMBER_NAME_CHECKED(UAnimGenControllerTrainingSettings, Behavior), UAnimGenControllerTrainingSettings::StaticClass());
 
-		if (TrainingSettings->Behavior && (
+		if (TrainingSettings->Database.Get() &&
+			TrainingSettings->Behavior && (
 			!TrainingSettings->ControlSchema.IsValid() || 
 			!TrainingSettings->ControlSchema.IsElementValid(TrainingSettings->ControlSchemaElement)))
 		{
@@ -444,7 +445,7 @@ namespace UE::AnimGen::Editor
 				];
 		}
 
-		if (Controller->IsValid() && Controller->AutoEncoder && Controller->AutoEncoder->TrainedContentHash != Controller->TrainedAutoEncoderContentHash)
+		if (Controller->IsValid() && Controller->AutoEncoder && Controller->AutoEncoder->GetContentHash() != Controller->TrainedAutoEncoderContentHash)
 		{
 			FDetailWidgetRow& NeedsReimportErrorRow = ControllerCategory.AddCustomRow(LOCTEXT("ModifiedWarning", "Warning"))
 				.WholeRowContent()
@@ -453,7 +454,7 @@ namespace UE::AnimGen::Editor
 						.Padding(FMargin(0.0f, 4.0f))
 						[
 							SNew(SWarningOrErrorBox)
-								.MessageStyle(EMessageStyle::Warning)
+								.MessageStyle(EMessageStyle::Error)
 								.Message(LOCTEXT("AutoEncoderModifiedWarningMessage",
 									"Auto-Encoder appears to have changed since Controller was last trained. Controller may no longer function correctly."))
 						]
@@ -474,7 +475,7 @@ namespace UE::AnimGen::Editor
 							.Padding(FMargin(0.0f, 4.0f))
 							[
 								SNew(SWarningOrErrorBox)
-									.MessageStyle(EMessageStyle::Warning)
+									.MessageStyle(EMessageStyle::Error)
 									.Message(LOCTEXT("BehaviorModifiedWarningMessage",
 										"Behavior appears to have changed since Controller was last trained."))
 							]
@@ -1615,7 +1616,7 @@ namespace UE::AnimGen::Editor
 
 		TOptional<FTimespan> GetEstimateTimeRemaining() const
 		{
-			if (Iteration > 200 && StartTrainingTimeIteration != -1)
+			if (StartTrainingTimeIteration != -1 && Iteration > StartTrainingTimeIteration + 250)
 			{
 				const FTimespan ElapsedTime = FDateTime::Now() - StartTrainingTime;
 				const int32 ElapsedIteration = Iteration - StartTrainingTimeIteration;
@@ -1641,7 +1642,7 @@ namespace UE::AnimGen::Editor
 
 					// Update Estimated Time
 
-					if (Iteration > 100 && StartTrainingTimeIteration == -1)
+					if (Iteration > 250 && StartTrainingTimeIteration == -1)
 					{
 						StartTrainingTimeIteration = Iteration;
 						StartTrainingTime = FDateTime::Now();
