AnimGen: Fixed some UX issues

#Synced-CL 52204516

diff --git a/Engine/Plugins/Experimental/Animation/AnimDatabase/Source/AnimDatabase/Private/AnimDatabase.cpp b/Engine/Plugins/Experimental/Animation/AnimDatabase/Source/AnimDatabase/Private/AnimDatabase.cpp
--- a/Engine/Plugins/Experimental/Animation/AnimDatabase/Source/AnimDatabase/Private/AnimDatabase.cpp
+++ b/Engine/Plugins/Experimental/Animation/AnimDatabase/Source/AnimDatabase/Private/AnimDatabase.cpp
@@ -2366,7 +2366,7 @@ void UAnimDatabaseQuery::UpdateQueryEntries()
 
 void UAnimDatabaseQuery::ForceRefresh()
 {
-	DatabaseContentHash = 0;
+	bForceRefresh = true;
 }
 
 void UAnimDatabaseQuery::DeleteSelectedFromDatabase()
@@ -2390,7 +2390,8 @@ bool UAnimDatabaseQuery::Update()
 	{
 		// Otherwise check for changes to various things to detect update
 
-		bool bQueryUpdateRequired = false;
+		bool bQueryUpdateRequired = bForceRefresh;
+		bForceRefresh = false;
 
 		// Check if Database has been updated
 
@@ -2514,7 +2515,7 @@ bool UAnimDatabaseQuery::Update()
 
 		return bQueryUpdateRequired;
 	}
-	
+		
 	return false;
 }
 
diff --git a/Engine/Plugins/Experimental/Animation/AnimDatabase/Source/AnimDatabase/Public/AnimDatabase.h b/Engine/Plugins/Experimental/Animation/AnimDatabase/Source/AnimDatabase/Public/AnimDatabase.h
--- a/Engine/Plugins/Experimental/Animation/AnimDatabase/Source/AnimDatabase/Public/AnimDatabase.h
+++ b/Engine/Plugins/Experimental/Animation/AnimDatabase/Source/AnimDatabase/Public/AnimDatabase.h
@@ -311,6 +311,9 @@ public:
 	UPROPERTY()
 	TObjectPtr<UAnimDatabase> Database;
 
+	/** Flag to force refresh the query during update */
+	bool bForceRefresh = false;
+
 	/** This hash is used to detect when the database might have changed and re-run the query */
 	int32 DatabaseContentHash = 0;
 
diff --git a/Engine/Plugins/Experimental/Animation/AnimGen/Source/AnimGen/Private/AnimGenAutoEncoder.cpp b/Engine/Plugins/Experimental/Animation/AnimGen/Source/AnimGen/Private/AnimGenAutoEncoder.cpp
--- a/Engine/Plugins/Experimental/Animation/AnimGen/Source/AnimGen/Private/AnimGenAutoEncoder.cpp
+++ b/Engine/Plugins/Experimental/Animation/AnimGen/Source/AnimGen/Private/AnimGenAutoEncoder.cpp
@@ -52,7 +52,7 @@ USkeleton* UAnimGenAutoEncoderTrainingSettings::GetSkeleton(bool& bInvalidSkelet
 
 void UAnimGenAutoEncoderTrainingSettings::ForceRefresh()
 {
-	DatabaseContentHash = 0;
+	bForceRefresh = true;
 }
 
 void UAnimGenAutoEncoderTrainingSettings::LoadDatabaseAsync()
@@ -145,7 +145,8 @@ bool UAnimGenAutoEncoderTrainingSettings::Update()
 	{
 		// Otherwise check for changes to various things to detect update
 
-		bool bQueryUpdateRequired = false;
+		bool bQueryUpdateRequired = bForceRefresh;
+		bForceRefresh = false;
 
 		// Check if Database has been updated
 
@@ -215,9 +216,19 @@ bool UAnimGenAutoEncoderTrainingSettings::Update()
 		return bQueryUpdateRequired;
 	}
 
-	const bool bUpdateRequired = DatabaseContentHash != 0;
-	DatabaseContentHash = 0;
-	return bUpdateRequired;
+	if (Database.IsNull() && (DatabaseContentHash != 0 || bForceRefresh))
+	{
+		bForceRefresh = false;
+		DatabaseContentHash = 0;
+		ActiveFrameRanges = nullptr;
+		QueryRanges = UAnimDatabaseFrameRangesLibrary::MakeEmptyFrameRanges();
+		UpdateQueryEntries();
+		ActiveDebugDrawer = nullptr;
+		SelectedRanges.Reset();
+		return true;
+	}
+
+	return false;
 }
 
 bool UAnimGenAutoEncoderTrainingSettings::UpdateFrameAttributeObjects()
diff --git a/Engine/Plugins/Experimental/Animation/AnimGen/Source/AnimGen/Private/AnimGenController.cpp b/Engine/Plugins/Experimental/Animation/AnimGen/Source/AnimGen/Private/AnimGenController.cpp
--- a/Engine/Plugins/Experimental/Animation/AnimGen/Source/AnimGen/Private/AnimGenController.cpp
+++ b/Engine/Plugins/Experimental/Animation/AnimGen/Source/AnimGen/Private/AnimGenController.cpp
@@ -24,7 +24,7 @@
 
 void UAnimGenControllerTrainingSettings::ForceRefresh()
 {
-	DatabaseContentHash = 0;
+	bForceRefresh = true;
 }
 
 void UAnimGenControllerTrainingSettings::LoadDatabaseAsync()
@@ -139,7 +139,8 @@ bool UAnimGenControllerTrainingSettings::Update()
 {
 	if (const UAnimDatabase* DatabaseObject = Database.Get())
 	{
-		bool bControlSetsUpdateRequired = false;
+		bool bControlSetsUpdateRequired = bForceRefresh;
+		bForceRefresh = false;
 
 		// Check if Database has been updated
 
@@ -228,9 +229,20 @@ bool UAnimGenControllerTrainingSettings::Update()
 		return bControlSetsUpdateRequired;
 	}
 
-	const bool bUpdateRequired = DatabaseContentHash != 0;
-	DatabaseContentHash = 0;
-	return bUpdateRequired;
+	if (Database.IsNull() && (DatabaseContentHash != 0 || bForceRefresh))
+	{
+		bForceRefresh = false;
+		DatabaseContentHash = 0;
+		ActiveFrameRanges = nullptr;
+		QueryRanges = UAnimDatabaseFrameRangesLibrary::MakeEmptyFrameRanges();
+		ActiveBehavior = nullptr;
+		ControlSets.Reset();
+		UpdateQueryEntries();
+		SelectedRanges.Reset();
+		return true;
+	}
+
+	return false;
 }
 
 void UAnimGenControllerTestingSettings::LoadDatabaseAsync()
diff --git a/Engine/Plugins/Experimental/Animation/AnimGen/Source/AnimGen/Public/AnimGenAutoEncoder.h b/Engine/Plugins/Experimental/Animation/AnimGen/Source/AnimGen/Public/AnimGenAutoEncoder.h
--- a/Engine/Plugins/Experimental/Animation/AnimGen/Source/AnimGen/Public/AnimGenAutoEncoder.h
+++ b/Engine/Plugins/Experimental/Animation/AnimGen/Source/AnimGen/Public/AnimGenAutoEncoder.h
@@ -318,9 +318,6 @@ public:
 	UPROPERTY(EditAnywhere, Category = "Data")
 	TSoftObjectPtr<UAnimDatabase> Database = nullptr;
 
-	/** This hash is used to detect when the database might have changed and re-run the query */
-	int32 DatabaseContentHash = 0;
-
 	/** Frame ranges in the database to use for training. Leave this as None to train on the full database */
 	UPROPERTY(EditAnywhere, Instanced, Category = "Data")
 	TObjectPtr<UAnimDatabaseFrameRangesFunction> FrameRanges = nullptr;
@@ -329,6 +326,14 @@ public:
 	UPROPERTY(Transient)
 	TObjectPtr<UAnimDatabaseFrameRangesFunction> ActiveFrameRanges;
 
+public:
+
+	/** Flag to force refresh the query during update */
+	bool bForceRefresh = false;
+
+	/** This hash is used to detect when the database might have changed and re-run the query */
+	int32 DatabaseContentHash = 0;
+
 	/** Query Ranges object for the current query */
 	FAnimDatabaseFrameRanges QueryRanges;
 
diff --git a/Engine/Plugins/Experimental/Animation/AnimGen/Source/AnimGen/Public/AnimGenController.h b/Engine/Plugins/Experimental/Animation/AnimGen/Source/AnimGen/Public/AnimGenController.h
--- a/Engine/Plugins/Experimental/Animation/AnimGen/Source/AnimGen/Public/AnimGenController.h
+++ b/Engine/Plugins/Experimental/Animation/AnimGen/Source/AnimGen/Public/AnimGenController.h
@@ -258,6 +258,9 @@ public:
 
 public:
 
+	/** Flag to force refresh the query during update */
+	bool bForceRefresh = false;
+
 	/** This hash is used to detect when the database might have changed and re-generate the control sets */
 	int32 DatabaseContentHash = 0;
 
diff --git a/Engine/Plugins/Experimental/Animation/AnimGen/Source/AnimGenEditor/Private/AnimGenEditorControllerToolkit.cpp b/Engine/Plugins/Experimental/Animation/AnimGen/Source/AnimGenEditor/Private/AnimGenEditorControllerToolkit.cpp
--- a/Engine/Plugins/Experimental/Animation/AnimGen/Source/AnimGenEditor/Private/AnimGenEditorControllerToolkit.cpp
+++ b/Engine/Plugins/Experimental/Animation/AnimGen/Source/AnimGenEditor/Private/AnimGenEditorControllerToolkit.cpp
@@ -427,22 +427,29 @@ namespace UE::AnimGen::Editor
 
 		// Check if Trained Database has been modified
 
-		if (Controller->IsValid() && Controller->TrainingSettings->Database.Get() && UAnimDatabaseFrameRangesLibrary::FrameRangesContentHash(
-			Controller->TrainingSettings->Database.Get(),
-			Controller->TrainingSettings->QueryRanges) != Controller->TrainedFrameRangesContentHash)
+		if (Controller->IsValid() && Controller->TrainingSettings->Database.Get())
 		{
-			FDetailWidgetRow& NeedsReimportErrorRow = ControllerCategory.AddCustomRow(LOCTEXT("ModifiedWarning", "Warning"))
-				.WholeRowContent()
-				[
-					SNew(SBox)
-						.Padding(FMargin(0.0f, 4.0f))
-						[
-							SNew(SWarningOrErrorBox)
-								.MessageStyle(EMessageStyle::Warning)
-								.Message(LOCTEXT("DatabaseModifiedWarningMessage",
-									"Training Data appears to have changed since Controller was last trained."))
-						]
-				];
+			const UAnimDatabase* TrainingDatabase = Controller->TrainingSettings->Database.Get();
+
+			const FAnimDatabaseFrameRanges TrainingFrameRanges = Controller->TrainingSettings->FrameRanges ?
+				UAnimDatabaseFrameRangesLibrary::MakeFrameRangesFromFunction(TrainingDatabase, UAnimDatabaseFrameRangesLibrary::MakeFrameRangesFromDatabase(TrainingDatabase), Controller->TrainingSettings->FrameRanges) :
+				UAnimDatabaseFrameRangesLibrary::MakeFrameRangesFromDatabase(TrainingDatabase);
+
+			if (UAnimDatabaseFrameRangesLibrary::FrameRangesContentHash(TrainingDatabase, TrainingFrameRanges) != Controller->TrainedFrameRangesContentHash)
+			{
+				FDetailWidgetRow& NeedsReimportErrorRow = ControllerCategory.AddCustomRow(LOCTEXT("ModifiedWarning", "Warning"))
+					.WholeRowContent()
+					[
+						SNew(SBox)
+							.Padding(FMargin(0.0f, 4.0f))
+							[
+								SNew(SWarningOrErrorBox)
+									.MessageStyle(EMessageStyle::Warning)
+									.Message(LOCTEXT("DatabaseModifiedWarningMessage",
+										"Training Data appears to have changed since Controller was last trained."))
+							]
+					];
+			}
 		}
 
 		if (Controller->IsValid() && Controller->AutoEncoder && Controller->AutoEncoder->GetContentHash() != Controller->TrainedAutoEncoderContentHash)
