AnimGen: Fixed shared memory guid string issues for mac

#Synced-CL 52297216

diff --git a/Engine/Plugins/Experimental/Animation/AnimGen/Source/AnimGenEditor/Private/AnimGenEditorAutoEncoderToolkit.cpp b/Engine/Plugins/Experimental/Animation/AnimGen/Source/AnimGenEditor/Private/AnimGenEditorAutoEncoderToolkit.cpp
--- a/Engine/Plugins/Experimental/Animation/AnimGen/Source/AnimGenEditor/Private/AnimGenEditorAutoEncoderToolkit.cpp
+++ b/Engine/Plugins/Experimental/Animation/AnimGen/Source/AnimGenEditor/Private/AnimGenEditorAutoEncoderToolkit.cpp
@@ -1193,15 +1193,15 @@ namespace UE::AnimGen::Editor
 				ConfigObject->SetStringField(TEXT("Device"), AutoEncoder->TrainingSettings->Device == EAnimGenTrainingDevice::GPU ? TEXT("GPU") : TEXT("CPU"));
 				ConfigObject->SetBoolField(TEXT("EnableTensorboard"), AutoEncoder->TrainingSettings->bEnableTensorboard);
 
-				ConfigObject->SetStringField(TEXT("ControlGuid"), *TrainProcessControls.Guid.ToString(EGuidFormats::DigitsWithHyphensInBraces));
-				ConfigObject->SetStringField(TEXT("EncoderGuid"), *EncoderNetworkData.Guid.ToString(EGuidFormats::DigitsWithHyphensInBraces));
-				ConfigObject->SetStringField(TEXT("DecoderGuid"), *DecoderNetworkData.Guid.ToString(EGuidFormats::DigitsWithHyphensInBraces));
+				ConfigObject->SetStringField(TEXT("ControlGuid"), *TrainProcessControls.GetPlatformGuidString());
+				ConfigObject->SetStringField(TEXT("EncoderGuid"), *EncoderNetworkData.GetPlatformGuidString());
+				ConfigObject->SetStringField(TEXT("DecoderGuid"), *DecoderNetworkData.GetPlatformGuidString());
 
-				ConfigObject->SetStringField(TEXT("RangeStartsGuid"), *RangeStarts.Guid.ToString(EGuidFormats::DigitsWithHyphensInBraces));
-				ConfigObject->SetStringField(TEXT("RangeLengthsGuid"), *RangeLengths.Guid.ToString(EGuidFormats::DigitsWithHyphensInBraces));
+				ConfigObject->SetStringField(TEXT("RangeStartsGuid"), *RangeStarts.GetPlatformGuidString());
+				ConfigObject->SetStringField(TEXT("RangeLengthsGuid"), *RangeLengths.GetPlatformGuidString());
 
-				ConfigObject->SetStringField(TEXT("PoseVectorsGuid"), *PoseVectors.Guid.ToString(EGuidFormats::DigitsWithHyphensInBraces));
-				ConfigObject->SetStringField(TEXT("PoseVectorWeightsGuid"), *PoseVectorWeights.Guid.ToString(EGuidFormats::DigitsWithHyphensInBraces));
+				ConfigObject->SetStringField(TEXT("PoseVectorsGuid"), *PoseVectors.GetPlatformGuidString());
+				ConfigObject->SetStringField(TEXT("PoseVectorWeightsGuid"), *PoseVectorWeights.GetPlatformGuidString());
 
 				FString JsonString;
 				TSharedRef<TJsonWriter<>> JsonWriter = TJsonWriterFactory<>::Create(&JsonString, 0);
diff --git a/Engine/Plugins/Experimental/Animation/AnimGen/Source/AnimGenEditor/Private/AnimGenEditorControllerToolkit.cpp b/Engine/Plugins/Experimental/Animation/AnimGen/Source/AnimGenEditor/Private/AnimGenEditorControllerToolkit.cpp
--- a/Engine/Plugins/Experimental/Animation/AnimGen/Source/AnimGenEditor/Private/AnimGenEditorControllerToolkit.cpp
+++ b/Engine/Plugins/Experimental/Animation/AnimGen/Source/AnimGenEditor/Private/AnimGenEditorControllerToolkit.cpp
@@ -1481,16 +1481,16 @@ namespace UE::AnimGen::Editor
 
 				// To let the Python use the arrays we created we pass it the Guids of the shared memory.
 
-				ConfigObject->SetStringField(TEXT("ControlGuid"), *TrainProcessControls.Guid.ToString(EGuidFormats::DigitsWithHyphensInBraces));
-				ConfigObject->SetStringField(TEXT("ControllerGuid"), *ControlEncoderNetworkData.Guid.ToString(EGuidFormats::DigitsWithHyphensInBraces));
-				ConfigObject->SetStringField(TEXT("LOD0Guid"), *LOD0NetworkData.Guid.ToString(EGuidFormats::DigitsWithHyphensInBraces));
-				ConfigObject->SetStringField(TEXT("LOD1Guid"), *LOD1NetworkData.Guid.ToString(EGuidFormats::DigitsWithHyphensInBraces));
-				ConfigObject->SetStringField(TEXT("LOD2Guid"), *LOD2NetworkData.Guid.ToString(EGuidFormats::DigitsWithHyphensInBraces));
-				ConfigObject->SetStringField(TEXT("RangeControlStartsGuid"), *RangeControlStarts.Guid.ToString(EGuidFormats::DigitsWithHyphensInBraces));
-				ConfigObject->SetStringField(TEXT("RangeEncodedStartsGuid"), *RangeEncodedStarts.Guid.ToString(EGuidFormats::DigitsWithHyphensInBraces));
-				ConfigObject->SetStringField(TEXT("RangeLengthsGuid"), *RangeLengths.Guid.ToString(EGuidFormats::DigitsWithHyphensInBraces));
-				ConfigObject->SetStringField(TEXT("ControlVectorsGuid"), *ControlVectors.Guid.ToString(EGuidFormats::DigitsWithHyphensInBraces));
-				ConfigObject->SetStringField(TEXT("EncodedVectorsGuid"), *EncodedVectors.Guid.ToString(EGuidFormats::DigitsWithHyphensInBraces));
+				ConfigObject->SetStringField(TEXT("ControlGuid"), *TrainProcessControls.GetPlatformGuidString());
+				ConfigObject->SetStringField(TEXT("ControllerGuid"), *ControlEncoderNetworkData.GetPlatformGuidString());
+				ConfigObject->SetStringField(TEXT("LOD0Guid"), *LOD0NetworkData.GetPlatformGuidString());
+				ConfigObject->SetStringField(TEXT("LOD1Guid"), *LOD1NetworkData.GetPlatformGuidString());
+				ConfigObject->SetStringField(TEXT("LOD2Guid"), *LOD2NetworkData.GetPlatformGuidString());
+				ConfigObject->SetStringField(TEXT("RangeControlStartsGuid"), *RangeControlStarts.GetPlatformGuidString());
+				ConfigObject->SetStringField(TEXT("RangeEncodedStartsGuid"), *RangeEncodedStarts.GetPlatformGuidString());
+				ConfigObject->SetStringField(TEXT("RangeLengthsGuid"), *RangeLengths.GetPlatformGuidString());
+				ConfigObject->SetStringField(TEXT("ControlVectorsGuid"), *ControlVectors.GetPlatformGuidString());
+				ConfigObject->SetStringField(TEXT("EncodedVectorsGuid"), *EncodedVectors.GetPlatformGuidString());
 
 				FString JsonString;
 				TSharedRef<TJsonWriter<>> JsonWriter = TJsonWriterFactory<>::Create(&JsonString, 0);
