AnimGen: Fixed threading issues on control error checking

#Synced-CL 52443113

diff --git a/Engine/Plugins/Experimental/Animation/AnimGen/Source/AnimGen/Private/AnimGenControl.cpp b/Engine/Plugins/Experimental/Animation/AnimGen/Source/AnimGen/Private/AnimGenControl.cpp
--- a/Engine/Plugins/Experimental/Animation/AnimGen/Source/AnimGen/Private/AnimGenControl.cpp
+++ b/Engine/Plugins/Experimental/Animation/AnimGen/Source/AnimGen/Private/AnimGenControl.cpp
@@ -1650,26 +1650,23 @@ FAnimGenControlObjectElement UAnimGenControls::MakeStructControlFromArrayViews(F
 		return FAnimGenControlObjectElement();
 	}
 
-	TArray<UE::Learning::Observation::FObjectElement, TInlineAllocator<16>> SubElements;
-	SubElements.Empty(Elements.Num());
-
-	for (const FAnimGenControlObjectElement& Element : Elements)
 	{
-		{
-			FReadScopeLock ObjectLock(Object.ObservationObject->Lock);
+		FWriteScopeLock ObjectLock(Object.ObservationObject->Lock);
+
+		TArray<UE::Learning::Observation::FObjectElement, TInlineAllocator<16>> SubElements;
+		SubElements.Empty(Elements.Num());
 
+		for (const FAnimGenControlObjectElement& Element : Elements)
+		{
 			if (!Object.ObservationObject->Object.IsValid(Element.ObjectElement))
 			{
 				UE_ANIMGEN_LOG_ONCE(LogAnimGen, Error, "MakeStructControlFromArrayViews: Invalid Control Object Element.");
 				return FAnimGenControlObjectElement();
 			}
-		}
 
-		SubElements.Add(Element.ObjectElement);
-	}
+			SubElements.Add(Element.ObjectElement);
+		}
 
-	{
-		FWriteScopeLock ObjectLock(Object.ObservationObject->Lock);
 		return { Object.ObservationObject->Object.CreateAnd({ ElementNames, SubElements }, Tag) };
 	}
 }
@@ -1683,17 +1680,14 @@ FAnimGenControlObjectElement UAnimGenControls::MakeExclusiveUnionControl(FAnimGe
 	}
 
 	{
-		FReadScopeLock ObjectLock(Object.ObservationObject->Lock);
+		FWriteScopeLock ObjectLock(Object.ObservationObject->Lock);
 
 		if (!Object.ObservationObject->Object.IsValid(Element.ObjectElement))
 		{
 			UE_ANIMGEN_LOG_ONCE(LogAnimGen, Error, "MakeExclusiveUnionControl: Invalid Control Object Element.");
 			return FAnimGenControlObjectElement();
 		}
-	}
 
-	{
-		FWriteScopeLock ObjectLock(Object.ObservationObject->Lock);
 		return { Object.ObservationObject->Object.CreateOrExclusive({ ElementName, Element.ObjectElement }, Tag) };
 	}
 }
@@ -1741,26 +1735,23 @@ FAnimGenControlObjectElement UAnimGenControls::MakeInclusiveUnionControlFromArra
 		return FAnimGenControlObjectElement();
 	}
 
-	TArray<UE::Learning::Observation::FObjectElement, TInlineAllocator<16>> SubElements;
-	SubElements.Empty(Elements.Num());
-
-	for (const FAnimGenControlObjectElement& Element : Elements)
 	{
-		{
-			FReadScopeLock ObjectLock(Object.ObservationObject->Lock);
+		FWriteScopeLock ObjectLock(Object.ObservationObject->Lock);
+
+		TArray<UE::Learning::Observation::FObjectElement, TInlineAllocator<16>> SubElements;
+		SubElements.Empty(Elements.Num());
 
+		for (const FAnimGenControlObjectElement& Element : Elements)
+		{
 			if (!Object.ObservationObject->Object.IsValid(Element.ObjectElement))
 			{
 				UE_ANIMGEN_LOG_ONCE(LogAnimGen, Error, "MakeInclusiveUnionControlFromArrayViews: Invalid Control Object Element.");
 				return FAnimGenControlObjectElement();
 			}
-		}
 
-		SubElements.Add(Element.ObjectElement);
-	}
+			SubElements.Add(Element.ObjectElement);
+		}
 
-	{
-		FWriteScopeLock ObjectLock(Object.ObservationObject->Lock);
 		return { Object.ObservationObject->Object.CreateOrInclusive({ ElementNames, SubElements }, Tag) };
 	}
 }
@@ -1783,26 +1774,23 @@ FAnimGenControlObjectElement UAnimGenControls::MakeStaticArrayControlFromArrayVi
 		UE_ANIMGEN_LOG_ONCE(LogAnimGen, Warning, "MakeStaticArrayControlFromArrayView: Creating zero-sized Static Array Control.");
 	}
 
-	TArray<UE::Learning::Observation::FObjectElement, TInlineAllocator<16>> SubElements;
-	SubElements.Empty(Elements.Num());
-
-	for (const FAnimGenControlObjectElement& Element : Elements)
 	{
-		{
-			FReadScopeLock ObjectLock(Object.ObservationObject->Lock);
+		FWriteScopeLock ObjectLock(Object.ObservationObject->Lock);
+
+		TArray<UE::Learning::Observation::FObjectElement, TInlineAllocator<16>> SubElements;
+		SubElements.Empty(Elements.Num());
 
+		for (const FAnimGenControlObjectElement& Element : Elements)
+		{
 			if (!Object.ObservationObject->Object.IsValid(Element.ObjectElement))
 			{
 				UE_ANIMGEN_LOG_ONCE(LogAnimGen, Error, "MakeStaticArrayControlFromArrayView: Invalid Control Object Element.");
 				return FAnimGenControlObjectElement();
 			}
-		}
 
-		SubElements.Add(Element.ObjectElement);
-	}
+			SubElements.Add(Element.ObjectElement);
+		}
 
-	{
-		FWriteScopeLock ObjectLock(Object.ObservationObject->Lock);
 		return { Object.ObservationObject->Object.CreateArray({ SubElements }, Tag) };
 	}
 }
@@ -1815,26 +1803,23 @@ FAnimGenControlObjectElement UAnimGenControls::MakeSetControl(FAnimGenControlObj
 		return FAnimGenControlObjectElement();
 	}
 
-	TArray<UE::Learning::Observation::FObjectElement, TInlineAllocator<16>> SubElements;
-	SubElements.Empty(Elements.Num());
-
-	for (const FAnimGenControlObjectElement& Element : Elements)
 	{
-		{
-			FReadScopeLock ObjectLock(Object.ObservationObject->Lock);
+		FWriteScopeLock ObjectLock(Object.ObservationObject->Lock);
+
+		TArray<UE::Learning::Observation::FObjectElement, TInlineAllocator<16>> SubElements;
+		SubElements.Empty(Elements.Num());
 
+		for (const FAnimGenControlObjectElement& Element : Elements)
+		{
 			if (!Object.ObservationObject->Object.IsValid(Element.ObjectElement))
 			{
 				UE_ANIMGEN_LOG_ONCE(LogAnimGen, Error, "MakeSetControl: Invalid Control Object Element.");
 				return FAnimGenControlObjectElement();
 			}
-		}
 
-		SubElements.Add(Element.ObjectElement);
-	}
+			SubElements.Add(Element.ObjectElement);
+		}
 
-	{
-		FWriteScopeLock ObjectLock(Object.ObservationObject->Lock);
 		return { Object.ObservationObject->Object.CreateSet({ SubElements }, Tag) };
 	}
 }
@@ -1852,26 +1837,23 @@ FAnimGenControlObjectElement UAnimGenControls::MakeSetControlFromArrayView(FAnim
 		return FAnimGenControlObjectElement();
 	}
 
-	TArray<UE::Learning::Observation::FObjectElement, TInlineAllocator<16>> SubElements;
-	SubElements.Empty(Elements.Num());
-
-	for (const FAnimGenControlObjectElement& Element : Elements)
 	{
-		{
-			FReadScopeLock ObjectLock(Object.ObservationObject->Lock);
+		FWriteScopeLock ObjectLock(Object.ObservationObject->Lock);
+
+		TArray<UE::Learning::Observation::FObjectElement, TInlineAllocator<16>> SubElements;
+		SubElements.Empty(Elements.Num());
 
+		for (const FAnimGenControlObjectElement& Element : Elements)
+		{
 			if (!Object.ObservationObject->Object.IsValid(Element.ObjectElement))
 			{
 				UE_ANIMGEN_LOG_ONCE(LogAnimGen, Error, "MakeSetControlFromArrayView: Invalid Control Object Element.");
 				return FAnimGenControlObjectElement();
 			}
-		}
 
-		SubElements.Add(Element.ObjectElement);
-	}
+			SubElements.Add(Element.ObjectElement);
+		}
 
-	{
-		FWriteScopeLock ObjectLock(Object.ObservationObject->Lock);
 		return { Object.ObservationObject->Object.CreateSet({ SubElements }, Tag) };
 	}
 }
@@ -1899,19 +1881,7 @@ FAnimGenControlObjectElement UAnimGenControls::MakeArrayControlFromArrayView(FAn
 
 	for (int32 ElementIdx = 0; ElementIdx < Elements.Num(); ElementIdx++)
 	{
-		const FAnimGenControlObjectElement Element = Elements[ElementIdx];
-
-		{
-			FReadScopeLock ObjectLock(Object.ObservationObject->Lock);
-
-			if (!Object.ObservationObject->Object.IsValid(Element.ObjectElement))
-			{
-				UE_ANIMGEN_LOG_ONCE(LogAnimGen, Error, "MakeArrayControlFromArrayView: Invalid Control Object Element.");
-				return FAnimGenControlObjectElement();
-			}
-		}
-
-		SubElements.Add(MakePairControl(Object, MakeCountControl(Object, ElementIdx, MaxNum), Element));
+		SubElements.Add(MakePairControl(Object, MakeCountControl(Object, ElementIdx, MaxNum), Elements[ElementIdx]));
 	}
 
 	return MakeSetControlFromArrayView(Object, SubElements, Tag);
@@ -1930,16 +1900,6 @@ FAnimGenControlObjectElement UAnimGenControls::MakeMapControl(FAnimGenControlObj
 
 	for (TPair<FAnimGenControlObjectElement, FAnimGenControlObjectElement> Item : Map)
 	{
-		{
-			FReadScopeLock ObjectLock(Object.ObservationObject->Lock);
-
-			if (!Object.ObservationObject->Object.IsValid(Item.Key.ObjectElement) || !Object.ObservationObject->Object.IsValid(Item.Value.ObjectElement))
-			{
-				UE_ANIMGEN_LOG_ONCE(LogAnimGen, Error, "MakeMapControl: Invalid Control Object Element.");
-				return FAnimGenControlObjectElement();
-			}
-		}
-
 		SubElements.Add(MakePairControl(Object, Item.Key, Item.Value));
 	}
 
@@ -1970,17 +1930,6 @@ FAnimGenControlObjectElement UAnimGenControls::MakeMapControlFromArrayViews(FAni
 
 	for (int32 ElementIdx = 0; ElementIdx < Keys.Num(); ElementIdx++)
 	{
-		{
-			FReadScopeLock ObjectLock(Object.ObservationObject->Lock);
-
-			if (!Object.ObservationObject->Object.IsValid(Keys[ElementIdx].ObjectElement) ||
-				!Object.ObservationObject->Object.IsValid(Values[ElementIdx].ObjectElement))
-			{
-				UE_ANIMGEN_LOG_ONCE(LogAnimGen, Error, "MakeMapControlFromArrayViews: Invalid Control Object Element.");
-				return FAnimGenControlObjectElement();
-			}
-		}
-
 		SubElements.Add(MakePairControl(Object, Keys[ElementIdx], Values[ElementIdx]));
 	}
 
@@ -2087,17 +2036,13 @@ FAnimGenControlObjectElement UAnimGenControls::MakeEncodingControl(FAnimGenContr
 	}
 
 	{
-		FReadScopeLock ObjectLock(Object.ObservationObject->Lock);
+		FWriteScopeLock ObjectLock(Object.ObservationObject->Lock);
 
 		if (!Object.ObservationObject->Object.IsValid(Element.ObjectElement))
 		{
 			UE_ANIMGEN_LOG_ONCE(LogAnimGen, Error, "MakeEncodingControl: Invalid Control Object Element.");
 			return FAnimGenControlObjectElement();
 		}
-	}
-
-	{
-		FWriteScopeLock ObjectLock(Object.ObservationObject->Lock);
 
 		return { Object.ObservationObject->Object.CreateEncoding({ Element.ObjectElement }, Tag) };
 	}
