AnimGen: Made tagging specified by sub-behavior

#changelist validated
#virtualized
#Synced-CL 52458764

diff --git a/Engine/Plugins/Experimental/Animation/AnimGen/Source/AnimGen/Private/AnimGenBehavior.cpp b/Engine/Plugins/Experimental/Animation/AnimGen/Source/AnimGen/Private/AnimGenBehavior.cpp
--- a/Engine/Plugins/Experimental/Animation/AnimGen/Source/AnimGen/Private/AnimGenBehavior.cpp
+++ b/Engine/Plugins/Experimental/Animation/AnimGen/Source/AnimGen/Private/AnimGenBehavior.cpp
@@ -42,136 +42,6 @@ void UAnimGenBehavior::MakeControlSets_Implementation(TArray<FAnimGenControlSet>
 	OutControlSets.Empty();
 }
 
-void UAnimGenBehaviorPreview::OnPreviewInitialize_Implementation(const UAnimGenController* Controller, const UAnimGenAutoEncoder* AutoEncoder) {}
-
-FAnimGenControlObjectElement UAnimGenBehaviorPreview::MakeControl_Implementation(FAnimGenControlObject& InControlObject, const FAnimDatabasePoseState& InPoseState) const
-{
-	return FAnimGenControlObjectElement();
-}
-
-bool UAnimGenBehaviorPreview::OnPreviewCheckReset_Implementation(FAnimGenBehaviorPreviewReset& OutResetData, const FAnimDatabasePoseState& State) { return false; }
-void UAnimGenBehaviorPreview::OnPreviewReset_Implementation(const FAnimDatabasePoseState& State) {}
-void UAnimGenBehaviorPreview::OnPreviewUpdate_Implementation(const FAnimDatabasePoseState& State, const float DeltaTime) {}
-
-void UAnimGenBehaviorPreview::OnPreviewDrawDebug_Implementation(
-	const FDebugDrawer& Drawer, 
-	const FAnimGenControlObject& InControlObject,
-	const FAnimGenControlObjectElement InControlObjectElement, 
-	const FAnimDatabasePoseState& State, 
-	const FLinearColor& IdentifierColor) {}
-
-void UAnimGenBehaviorPreview::OnPreviewInputLeftStickAxis_Implementation(const FVector Axis) {}
-void UAnimGenBehaviorPreview::OnPreviewInputRightStickAxis_Implementation(const FVector Axis) {}
-void UAnimGenBehaviorPreview::OnPreviewInputLeftTriggerAxis_Implementation(const float Axis) {}
-void UAnimGenBehaviorPreview::OnPreviewInputRightTriggerAxis_Implementation(const float Axis) {}
-void UAnimGenBehaviorPreview::OnPreviewInputMouseMoved_Implementation(const FVector MouseGroundLocation) {}
-void UAnimGenBehaviorPreview::OnPreviewInputLeftMousePressed_Implementation(const FVector MouseGroundLocation) {}
-void UAnimGenBehaviorPreview::OnPreviewInputLeftMouseReleased_Implementation(const FVector MouseGroundLocation) {}
-void UAnimGenBehaviorPreview::OnPreviewInputRightMousePressed_Implementation(const FVector MouseGroundLocation) {}
-void UAnimGenBehaviorPreview::OnPreviewInputRightMouseReleased_Implementation(const FVector MouseGroundLocation) {}
-void UAnimGenBehaviorPreview::OnPreviewInputKeyPressed_Implementation(const FKey Key) {}
-void UAnimGenBehaviorPreview::OnPreviewInputKeyReleased_Implementation(const FKey Key) {}
-
-void UAnimGenBehaviorPreview::GetPreviewCameraLocationAndRotation(FVector& OutCameraLocation, FRotator& OutCameraRotation) const
-{
-	OutCameraLocation = PreviewCameraLocation;
-	OutCameraRotation = PreviewCameraRotation;
-}
-
-FRotator UAnimGenBehaviorPreview::GetPreviewCameraYawRotation() const
-{
-	return FRotator(0.0f, PreviewCameraRotation.Yaw, 0.0f);
-}
-
-void UAnimGenBehaviorPreview::UpdatePreviewCamera(
-	const float YawVelocity,
-	const float PitchVelocity,
-	const float DistanceVelocity,
-	const float DeltaTime,
-	const bool bInvertPitch)
-{
-	PreviewCameraYaw += DeltaTime * YawVelocity;
-	PreviewCameraPitch = FMath::Clamp(PreviewCameraPitch + DeltaTime * (bInvertPitch ? -PitchVelocity : PitchVelocity), 0.0f, 80.0f);
-	PreviewCameraDistance = FMath::Clamp(PreviewCameraDistance + DeltaTime * DistanceVelocity, 10.0f, 1000.0);
-}
-
-void UAnimGenBehaviorPreview::UpdatePreviewCameraFromAxisAndShoulders(
-	const FVector Axis,
-	const bool bLeftShoulder,
-	const bool bRightShoulder,
-	const float DeltaTime,
-	const float RotationSpeed,
-	const float DistanceSpeed,
-	const bool bInvertPitch,
-	const float Deadzone,
-	const float Sensitivity)
-{
-	FVector OutAxis = FVector::ZeroVector;
-	FVector OutDirection = FVector::ZeroVector;
-	float OutLength = 0.0f;
-	bool bOutNotInDeadzone = false;
-	UAnimGenBehaviorLibrary::ApplyStickDeadzoneAndSensitivity(
-		OutAxis,
-		OutDirection,
-		OutLength,
-		bOutNotInDeadzone,
-		Axis,
-		Deadzone,
-		Sensitivity);
-
-	UpdatePreviewCamera(
-		bOutNotInDeadzone ? RotationSpeed * -OutAxis.Y : 0.0f,
-		bOutNotInDeadzone ? RotationSpeed * OutAxis.X : 0.0f,
-		bLeftShoulder ? DistanceSpeed : bRightShoulder ? -DistanceSpeed : 0.0f,
-		DeltaTime,
-		bInvertPitch);
-}
-
-FVector UAnimGenBehaviorPreview::GetPreviewLeftStickAxis() const
-{
-	return PreviewLeftStick;
-}
-
-FVector UAnimGenBehaviorPreview::GetPreviewRightStickAxis() const
-{
-	return PreviewRightStick;
-}
-
-FVector UAnimGenBehaviorPreview::GetPreviewCameraRelativeLeftStickAxis() const
-{
-	return GetPreviewCameraYawRotation().RotateVector(PreviewLeftStick);
-}
-
-FVector UAnimGenBehaviorPreview::GetPreviewCameraRelativeRightStickAxis() const
-{
-	return GetPreviewCameraYawRotation().RotateVector(PreviewRightStick);
-}
-
-float UAnimGenBehaviorPreview::GetPreviewLeftTriggerAxis() const
-{
-	return PreviewLeftTrigger;
-}
-
-float UAnimGenBehaviorPreview::GetPreviewRightTriggerAxis() const
-{
-	return PreviewRightTrigger;
-}
-
-float UAnimGenBehaviorPreview::GetPreviewLeftTriggerAxisWithSensitivity(const float Sensitivity) const
-{
-	return UAnimGenBehaviorLibrary::ApplyTriggerSensitivity(PreviewLeftTrigger, Sensitivity);
-}
-
-float UAnimGenBehaviorPreview::GetPreviewRightTriggerAxisWithSensitivity(const float Sensitivity) const
-{
-	return UAnimGenBehaviorLibrary::ApplyTriggerSensitivity(PreviewRightTrigger, Sensitivity);
-}
-
-FVector UAnimGenBehaviorPreview::GetPreviewMouseGroundLocation() const
-{
-	return PreviewMouseGroundLocation;
-}
-
 FAnimGenControlSchemaElement UAnimGenBehavior_Uncontrolled::SpecifyControl_Implementation(FAnimGenControlSchema& InControlSchema) const
 {
 	return UAnimGenControls::SpecifyNullControl(InControlSchema);
@@ -182,86 +52,9 @@ void UAnimGenBehavior_Uncontrolled::MakeControlSets_Implementation(TArray<FAnimG
 	OutControlSets.Add(UAnimGenControls::MakeNullControlSet(InControlObject, InFrameRanges));
 }
 
-FAnimGenControlSchemaElement UAnimGenBehavior_Idle::SpecifyControl_Implementation(FAnimGenControlSchema& InControlSchema) const
-{
-	const int32 TagNum = TagRanges.Num();
-
-	TArray<FName, TInlineAllocator<32>> TagNames;
-	TagNames.Reserve(TagNum);
-
-	for (int32 TagIdx = 0; TagIdx < TagNum; TagIdx++)
-	{
-		TagNames.AddUnique(TagRanges[TagIdx].Name);
-	}
-
-	return UAnimGenControls::SpecifyNamedInclusiveDiscreteControlFromArrayView(InControlSchema, TagNames);
-}
-
-void UAnimGenBehavior_Idle::MakeControlSets_Implementation(TArray<FAnimGenControlSet>& OutControlSets, FAnimGenControlObject& InControlObject, const UAnimDatabase* InDatabase, const FAnimDatabaseFrameRanges& InFrameRanges) const
-{
-	const int32 TagNum = TagRanges.Num();
-
-	TArray<FName, TInlineAllocator<32>> TagNames;
-	TArray<FAnimDatabaseFrameRanges, TInlineAllocator<32>> TagFrameRanges;
-	TagNames.Reserve(TagNum);
-	TagFrameRanges.Reserve(TagNum);
-
-	for (int32 TagIdx = 0; TagIdx < TagNum; TagIdx++)
-	{
-		if (TagNames.Contains(TagRanges[TagIdx].Name)) { continue; }
-
-		const FAnimDatabaseFrameRanges Ranges = UAnimDatabaseFrameRangesLibrary::MakeFrameRangesFromFunction(InDatabase, InFrameRanges, TagRanges[TagIdx].FrameRanges);
-
-		if (Ranges.IsValid() && !Ranges.FrameRangeSet->IsEmpty())
-		{
-			TagNames.Add(TagRanges[TagIdx].Name);
-			TagFrameRanges.Add(Ranges);
-		}
-	}
-
-	OutControlSets.Add(UAnimGenControls::MakeNamedInclusiveDiscreteControlSetFromFrameRangesArrayViews(InControlObject, InFrameRanges, TagNames, TagFrameRanges));
-}
-
-void UAnimGenBehavior_Idle::DrawDebugControl_Implementation(
-	const FDebugDrawer& Drawer,
-	const FAnimGenControlObject& InControlObject,
-	const FAnimGenControlObjectElement InControlObjectElement,
-	const FAnimDatabasePoseState& InPoseState,
-	const UAnimDatabase* InDatabase,
-	const UAnimGenAutoEncoder* InAutoEncoder,
-	const int32 SequenceIdx,
-	const float SequenceTime,
-	const int32 RangeStart,
-	const int32 RangeLength,
-	const FLinearColor& IdentifierColor) const
-{
-	const FTransform RootTransform = UAnimDatabasePoseStateLibrary::PoseStateRootTransform(InPoseState);
-
-	UAnimGenControls::DrawDebugNamedInclusiveDiscreteControl(Drawer, InControlObject, InControlObjectElement, TEXT(""), RootTransform.GetLocation() + FVector(0.0, 0.0, 220.0f), FRotator(0.0f, -90.0f, 0.0f), 10.0f, IdentifierColor);
-}
-
 FAnimGenControlSchemaElement UAnimGenBehavior_TrajectoryFollow::SpecifyControl_Implementation(FAnimGenControlSchema& InControlSchema) const
 {
-	const int32 TagNum = TagRanges.Num();
-
-	TArray<FName, TInlineAllocator<32>> TagNames;
-	TagNames.Reserve(TagNum);
-
-	for (int32 TagIdx = 0; TagIdx < TagNum; TagIdx++)
-	{
-		TagNames.AddUnique(TagRanges[TagIdx].Name);
-	}
-		
-
-	return UAnimGenControls::SpecifyStructControlFromArrayViews(InControlSchema,
-		{
-			TEXT("Trajectory"),
-			TEXT("Tags")
-		},
-		{
-			UAnimGenControls::SpecifyTrajectoryControl(InControlSchema, TrajectorySampleNum),
-			UAnimGenControls::SpecifyNamedInclusiveDiscreteControlFromArrayView(InControlSchema, TagNames),
-		});
+	return UAnimGenControls::SpecifyTrajectoryControl(InControlSchema, TrajectorySampleNum);
 }
 
 void UAnimGenBehavior_TrajectoryFollow::MakeControlSets_Implementation(TArray<FAnimGenControlSet>& OutControlSets, FAnimGenControlObject& InControlObject, const UAnimDatabase* InDatabase, const FAnimDatabaseFrameRanges& InFrameRanges) const
@@ -274,35 +67,7 @@ void UAnimGenBehavior_TrajectoryFollow::MakeControlSets_Implementation(TArray<FA
 			FMath::RoundToInt(TrajectoryPastTime * InDatabase->GetFrameRate().AsDecimal()),
 			FMath::RoundToInt(TrajectoryFutureTime * InDatabase->GetFrameRate().AsDecimal())));
 
-	const int32 TagNum = TagRanges.Num();
-
-	TArray<FName, TInlineAllocator<32>> TagNames;
-	TArray<FAnimDatabaseFrameRanges, TInlineAllocator<32>> TagFrameRanges;
-	TagNames.Reserve(TagNum);
-	TagFrameRanges.Reserve(TagNum);
-
-	for (int32 TagIdx = 0; TagIdx < TagNum; TagIdx++)
-	{
-		if (TagNames.Contains(TagRanges[TagIdx].Name)) { continue; }
-
-		const FAnimDatabaseFrameRanges Ranges = UAnimDatabaseFrameRangesLibrary::MakeFrameRangesFromFunction(InDatabase, TrimmedRanges, TagRanges[TagIdx].FrameRanges);
-
-		if (Ranges.IsValid() && !Ranges.FrameRangeSet->IsEmpty())
-		{
-			TagNames.Add(TagRanges[TagIdx].Name);
-			TagFrameRanges.Add(Ranges);
-		}
-	}
-
-	OutControlSets.Add(UAnimGenControls::MakeStructControlSetFromArrayViews(InControlObject,
-		{
-			TEXT("Trajectory"),
-			TEXT("Tags")
-		},
-		{
-			UAnimGenControls::MakeTrajectoryControlSetFromDatabase(InControlObject, InDatabase, TrimmedRanges, TrajectorySampleNum, TrajectoryFutureTime, TrajectoryPastTime, TrajectoryForwardVector),
-			UAnimGenControls::MakeNamedInclusiveDiscreteControlSetFromFrameRangesArrayViews(InControlObject, TrimmedRanges, TagNames, TagFrameRanges),
-		}));
+	OutControlSets.Add(UAnimGenControls::MakeTrajectoryControlSetFromDatabase(InControlObject, InDatabase, TrimmedRanges, TrajectorySampleNum, TrajectoryFutureTime, TrajectoryPastTime, TrajectoryForwardVector));
 }
 
 void UAnimGenBehavior_TrajectoryFollow::DrawDebugControl_Implementation(
@@ -318,74 +83,34 @@ void UAnimGenBehavior_TrajectoryFollow::DrawDebugControl_Implementation(
 	const int32 RangeLength,
 	const FLinearColor& IdentifierColor) const
 {
-	FAnimGenControlObjectElement TrajectoryControl, TagsControl;
-	UAnimGenControls::GetStructControlElement(TrajectoryControl, InControlObject, InControlObjectElement, TEXT("Trajectory"));
-	UAnimGenControls::GetStructControlElement(TagsControl, InControlObject, InControlObjectElement, TEXT("Tags"));
-
 	const FTransform RootTransform = UAnimDatabasePoseStateLibrary::PoseStateRootTransform(InPoseState);
-
-	UAnimGenControls::DrawDebugTrajectoryControl(Drawer, InControlObject, TrajectoryControl, RootTransform, IdentifierColor, 25.0f, 5.0f, 1.0f, 1.0f);
-	UAnimGenControls::DrawDebugNamedInclusiveDiscreteControl(Drawer, InControlObject, TagsControl, TEXT(""), RootTransform.GetLocation() + FVector(0.0, 0.0, 220.0f), FRotator(0.0f, -90.0f, 0.0f), 10.0f, IdentifierColor);
+	UAnimGenControls::DrawDebugTrajectoryControl(Drawer, InControlObject, InControlObjectElement, RootTransform, IdentifierColor, 25.0f, 5.0f, 1.0f, 1.0f);
 }
 
 
 FAnimGenControlSchemaElement UAnimGenBehavior_VelocityFollow::SpecifyControl_Implementation(FAnimGenControlSchema& InControlSchema) const
 {
-	const int32 TagNum = TagRanges.Num();
-
-	TArray<FName, TInlineAllocator<32>> TagNames;
-	TagNames.Reserve(TagNum);
-
-	for (int32 TagIdx = 0; TagIdx < TagNum; TagIdx++)
-	{
-		TagNames.AddUnique(TagRanges[TagIdx].Name);
-	}
-
 	return UAnimGenControls::SpecifyStructControlFromArrayViews(InControlSchema,
 		{
 			TEXT("LinearVelocity"),
 			TEXT("AngularVelocity"),
-			TEXT("Tags")
 		},
 		{
 			UAnimGenControls::SpecifyLinearVelocityControl(InControlSchema),
 			UAnimGenControls::SpecifyAngularVelocityControl(InControlSchema),
-			UAnimGenControls::SpecifyNamedInclusiveDiscreteControlFromArrayView(InControlSchema, TagNames),
 		});
 }
 
 void UAnimGenBehavior_VelocityFollow::MakeControlSets_Implementation(TArray<FAnimGenControlSet>& OutControlSets, FAnimGenControlObject& InControlObject, const UAnimDatabase* InDatabase, const FAnimDatabaseFrameRanges& InFrameRanges) const
 {
-	const int32 TagNum = TagRanges.Num();
-
-	TArray<FName, TInlineAllocator<32>> TagNames;
-	TArray<FAnimDatabaseFrameRanges, TInlineAllocator<32>> TagFrameRanges;
-	TagNames.Reserve(TagNum);
-	TagFrameRanges.Reserve(TagNum);
-
-	for (int32 TagIdx = 0; TagIdx < TagNum; TagIdx++)
-	{
-		if (TagNames.Contains(TagRanges[TagIdx].Name)) { continue; }
-
-		const FAnimDatabaseFrameRanges Ranges = UAnimDatabaseFrameRangesLibrary::MakeFrameRangesFromFunction(InDatabase, InFrameRanges, TagRanges[TagIdx].FrameRanges);
-
-		if (Ranges.IsValid() && !Ranges.FrameRangeSet->IsEmpty())
-		{
-			TagNames.Add(TagRanges[TagIdx].Name);
-			TagFrameRanges.Add(Ranges);
-		}
-	}
-
 	OutControlSets.Add(UAnimGenControls::MakeStructControlSetFromArrayViews(InControlObject,
 		{
 			TEXT("LinearVelocity"),
 			TEXT("AngularVelocity"),
-			TEXT("Tags")
 		},
 		{
 			UAnimGenControls::MakeRootLinearVelocityControlSetFromDatabase(InControlObject, InDatabase, InFrameRanges),
 			UAnimGenControls::MakeRootAngularVelocityControlSetFromDatabase(InControlObject, InDatabase, InFrameRanges),
-			UAnimGenControls::MakeNamedInclusiveDiscreteControlSetFromFrameRangesArrayViews(InControlObject, InFrameRanges, TagNames, TagFrameRanges),
 		}));
 }
 
@@ -402,91 +127,52 @@ void UAnimGenBehavior_VelocityFollow::DrawDebugControl_Implementation(
 	const int32 RangeLength,
 	const FLinearColor& IdentifierColor) const
 {
-	FAnimGenControlObjectElement LinearVelocityControl, AngularVelocityControl, TagsControl;
+	FAnimGenControlObjectElement LinearVelocityControl, AngularVelocityControl;
 	UAnimGenControls::GetStructControlElement(LinearVelocityControl, InControlObject, InControlObjectElement, TEXT("LinearVelocity"));
 	UAnimGenControls::GetStructControlElement(AngularVelocityControl, InControlObject, InControlObjectElement, TEXT("AngularVelocity"));
-	UAnimGenControls::GetStructControlElement(TagsControl, InControlObject, InControlObjectElement, TEXT("Tags"));
 
 	const FTransform RootTransform = UAnimDatabasePoseStateLibrary::PoseStateRootTransform(InPoseState);
 
 	UAnimGenControls::DrawDebugLinearVelocityControl(Drawer, InControlObject, LinearVelocityControl, RootTransform, RootTransform.GetLocation() + FVector(0.0f, 0.0f, 10.0f), IdentifierColor, 1.0f, 1.0f);
 	UAnimGenControls::DrawDebugAngularVelocityControl(Drawer, InControlObject, AngularVelocityControl, RootTransform, RootTransform.GetLocation() + FVector(0.0f, 0.0f, 10.0f), IdentifierColor, 1.0f, 1.0f);
-	UAnimGenControls::DrawDebugNamedInclusiveDiscreteControl(Drawer, InControlObject, TagsControl, TEXT(""), RootTransform.GetLocation() + FVector(0.0, 0.0, 220.0f), FRotator(0.0f, -90.0f, 0.0f), 10.0f, IdentifierColor);
 }
 
 FAnimGenControlSchemaElement UAnimGenBehavior_MoveToTarget::SpecifyControl_Implementation(FAnimGenControlSchema& InControlSchema) const
 {
-	const int32 TagNum = TagRanges.Num();
-
-	TArray<FName, TInlineAllocator<32>> TagNames;
-	TagNames.Reserve(TagNum);
-
-	for (int32 TagIdx = 0; TagIdx < TagNum; TagIdx++)
-	{
-		TagNames.AddUnique(TagRanges[TagIdx].Name);
-	}
-
 	return UAnimGenControls::SpecifyStructControlFromArrayViews(InControlSchema,
 		{
 			TEXT("Location"),
 			TEXT("Direction"),
-			TEXT("Tags")
 		},
 		{
 			UAnimGenControls::SpecifyLocationControl(InControlSchema),
 			UAnimGenControls::SpecifyOptionalControl(InControlSchema, UAnimGenControls::SpecifyDirectionControl(InControlSchema)),
-			UAnimGenControls::SpecifyNamedInclusiveDiscreteControlFromArrayView(InControlSchema, TagNames),
 		});
 }
 
 void UAnimGenBehavior_MoveToTarget::MakeControlSets_Implementation(TArray<FAnimGenControlSet>& OutControlSets, FAnimGenControlObject& InControlObject, const UAnimDatabase* InDatabase, const FAnimDatabaseFrameRanges& InFrameRanges) const
 {
-	const int32 TagNum = TagRanges.Num();
-
-	TArray<FName, TInlineAllocator<32>> TagNames;
-	TArray<FAnimDatabaseFrameRanges, TInlineAllocator<32>> TagFrameRanges;
-	TagNames.Reserve(TagNum);
-	TagFrameRanges.Reserve(TagNum);
-
-	for (int32 TagIdx = 0; TagIdx < TagNum; TagIdx++)
-	{
-		if (TagNames.Contains(TagRanges[TagIdx].Name)) { continue; }
-
-		const FAnimDatabaseFrameRanges Ranges = UAnimDatabaseFrameRangesLibrary::MakeFrameRangesFromFunction(InDatabase, InFrameRanges, TagRanges[TagIdx].FrameRanges);
-
-		if (Ranges.IsValid() && !Ranges.FrameRangeSet->IsEmpty())
-		{
-			TagNames.Add(TagRanges[TagIdx].Name);
-			TagFrameRanges.Add(Ranges);
-		}
-	}
-
 	const FAnimGenControlSet LocationControlSet = UAnimGenControls::MakeRootLocationAtRangeEndControlSetFromDatabase(InControlObject, InDatabase, InFrameRanges);
 	const FAnimGenControlSet DirectionControlSet = UAnimGenControls::MakeRootDirectionAtRangeEndControlSetFromDatabase(InControlObject, InDatabase, InFrameRanges);
-	const FAnimGenControlSet TagsControlSet = UAnimGenControls::MakeNamedInclusiveDiscreteControlSetFromFrameRangesArrayViews(InControlObject, InFrameRanges, TagNames, TagFrameRanges);
 
 	OutControlSets.Add(UAnimGenControls::MakeStructControlSetFromArrayViews(InControlObject,
 		{
 			TEXT("Location"),
 			TEXT("Direction"),
-			TEXT("Tags")
 		},
 		{
 			LocationControlSet,
 			UAnimGenControls::MakeOptionalValidControlSet(InControlObject, DirectionControlSet),
-			TagsControlSet,
 		}));
 
 	OutControlSets.Add(UAnimGenControls::MakeStructControlSetFromArrayViews(InControlObject,
 		{
 			TEXT("Location"),
 			TEXT("Direction"),
-			TEXT("Tags")
 		},
 		{
 			LocationControlSet,
 			UAnimGenControls::MakeOptionalNullControlSet(InControlObject, InFrameRanges),
-			TagsControlSet,
 		}));
 }
 
@@ -503,10 +189,9 @@ void UAnimGenBehavior_MoveToTarget::DrawDebugControl_Implementation(
 	const int32 RangeLength,
 	const FLinearColor& IdentifierColor) const
 {
-	FAnimGenControlObjectElement LocationControl, OptionalDirectionControl, TagsControl;
+	FAnimGenControlObjectElement LocationControl, OptionalDirectionControl;
 	UAnimGenControls::GetStructControlElement(LocationControl, InControlObject, InControlObjectElement, TEXT("Location"));
 	UAnimGenControls::GetStructControlElement(OptionalDirectionControl, InControlObject, InControlObjectElement, TEXT("Direction"));
-	UAnimGenControls::GetStructControlElement(TagsControl, InControlObject, InControlObjectElement, TEXT("Tags"));
 
 	const FTransform RootTransform = UAnimDatabasePoseStateLibrary::PoseStateRootTransform(InPoseState);
 
@@ -521,32 +206,17 @@ void UAnimGenBehavior_MoveToTarget::DrawDebugControl_Implementation(
 		UAnimGenControls::GetLocationControl(Location, InControlObject, LocationControl, RootTransform);
 		UAnimGenControls::DrawDebugDirectionControl(Drawer, InControlObject, DirectionControl, RootTransform, Location + FVector(0.0f, 0.0, 1.0f), IdentifierColor, 25.0f, 1.0f, 1.0f);
 	}
-
-	UAnimGenControls::DrawDebugNamedInclusiveDiscreteControl(Drawer, InControlObject, TagsControl, TEXT(""), RootTransform.GetLocation() + FVector(0.0, 0.0, 220.0f), FRotator(0.0f, -90.0f, 0.0f), 10.0f, IdentifierColor);
 }
 
 FAnimGenControlSchemaElement UAnimGenBehavior_TrajectoryInteraction::SpecifyControl_Implementation(FAnimGenControlSchema& InControlSchema) const
 {
-	const int32 TagNum = TagRanges.Num();
-
-	TArray<FName, TInlineAllocator<32>> TagNames;
-	TagNames.Reserve(TagNum);
-
-	for (int32 TagIdx = 0; TagIdx < TagNum; TagIdx++)
-	{
-		TagNames.AddUnique(TagRanges[TagIdx].Name);
-	}
-
-
 	return UAnimGenControls::SpecifyStructControlFromArrayViews(InControlSchema,
 		{
 			TEXT("Trajectory"),
-			TEXT("Tags"),
 			TEXT("Interaction")
 		},
 		{
 			UAnimGenControls::SpecifyTrajectoryControl(InControlSchema, TrajectorySampleNum),
-			UAnimGenControls::SpecifyNamedInclusiveDiscreteControlFromArrayView(InControlSchema, TagNames),
 			UAnimGenControls::SpecifyOptionalControl(InControlSchema, UAnimGenControls::SpecifyStructControlFromArrayViews(InControlSchema,
 				{
 					TEXT("Location"),
@@ -569,26 +239,6 @@ void UAnimGenBehavior_TrajectoryInteraction::MakeControlSets_Implementation(TArr
 			FMath::RoundToInt(TrajectoryPastTime * InDatabase->GetFrameRate().AsDecimal()),
 			FMath::RoundToInt(TrajectoryFutureTime * InDatabase->GetFrameRate().AsDecimal())));
 
-	const int32 TagNum = TagRanges.Num();
-
-	TArray<FName, TInlineAllocator<32>> TagNames;
-	TArray<FAnimDatabaseFrameRanges, TInlineAllocator<32>> TagFrameRanges;
-	TagNames.Reserve(TagNum);
-	TagFrameRanges.Reserve(TagNum);
-
-	for (int32 TagIdx = 0; TagIdx < TagNum; TagIdx++)
-	{
-		if (TagNames.Contains(TagRanges[TagIdx].Name)) { continue; }
-
-		const FAnimDatabaseFrameRanges Ranges = UAnimDatabaseFrameRangesLibrary::MakeFrameRangesFromFunction(InDatabase, TrimmedRanges, TagRanges[TagIdx].FrameRanges);
-
-		if (Ranges.IsValid() && !Ranges.FrameRangeSet->IsEmpty())
-		{
-			TagNames.Add(TagRanges[TagIdx].Name);
-			TagFrameRanges.Add(Ranges);
-		}
-	}
-
 	FAnimGenControlSet InteractionFacingControlSet;
 
 	if (bUseMirroredInteractionForwardVector)
@@ -610,12 +260,10 @@ void UAnimGenBehavior_TrajectoryInteraction::MakeControlSets_Implementation(TArr
 	OutControlSets.Add(UAnimGenControls::MakeStructControlSetFromArrayViews(InControlObject,
 		{
 			TEXT("Trajectory"),
-			TEXT("Tags"),
 			TEXT("Interaction")
 		},
 		{
 			UAnimGenControls::MakeTrajectoryControlSetFromDatabase(InControlObject, InDatabase, TrimmedRanges, TrajectorySampleNum, TrajectoryFutureTime, TrajectoryPastTime, TrajectoryForwardVector),
-			UAnimGenControls::MakeNamedInclusiveDiscreteControlSetFromFrameRangesArrayViews(InControlObject, TrimmedRanges, TagNames, TagFrameRanges),
 			UAnimGenControls::MakeOptionalControlSet(InControlObject, UAnimGenControls::MakeStructControlSetFromArrayViews(InControlObject, 
 					{
 						TEXT("Location"),
@@ -644,15 +292,13 @@ void UAnimGenBehavior_TrajectoryInteraction::DrawDebugControl_Implementation(
 	const int32 RangeLength,
 	const FLinearColor& IdentifierColor) const
 {
-	FAnimGenControlObjectElement TrajectoryControl, TagsControl, InteractionControl;
+	FAnimGenControlObjectElement TrajectoryControl, InteractionControl;
 	UAnimGenControls::GetStructControlElement(TrajectoryControl, InControlObject, InControlObjectElement, TEXT("Trajectory"));
-	UAnimGenControls::GetStructControlElement(TagsControl, InControlObject, InControlObjectElement, TEXT("Tags"));
 	UAnimGenControls::GetStructControlElement(InteractionControl, InControlObject, InControlObjectElement, TEXT("Interaction"));
 
 	const FTransform RootTransform = UAnimDatabasePoseStateLibrary::PoseStateRootTransform(InPoseState);
 
 	UAnimGenControls::DrawDebugTrajectoryControl(Drawer, InControlObject, TrajectoryControl, RootTransform, IdentifierColor, 25.0f, 5.0f, 1.0f, 1.0f);
-	UAnimGenControls::DrawDebugNamedInclusiveDiscreteControl(Drawer, InControlObject, TagsControl, TEXT(""), RootTransform.GetLocation() + FVector(0.0, 0.0, 220.0f), FRotator(0.0f, -90.0f, 0.0f), 10.0f, IdentifierColor);
 
 	EAnimGenOptionalControl Option;
 	FAnimGenControlObjectElement InteractionControlStruct;
@@ -672,7 +318,8 @@ void UAnimGenBehavior_TrajectoryInteraction::DrawDebugControl_Implementation(
 	}
 }
 
-FAnimGenControlSchemaElement UAnimGenBehavior_Multi::SpecifyControl_Implementation(FAnimGenControlSchema& InControlSchema) const
+
+FAnimGenControlSchemaElement UAnimGenBehavior_Tagged::SpecifyControl_Implementation(FAnimGenControlSchema& InControlSchema) const
 {
 	const int32 TagNum = TagRanges.Num();
 
@@ -684,32 +331,18 @@ FAnimGenControlSchemaElement UAnimGenBehavior_Multi::SpecifyControl_Implementati
 		TagNames.AddUnique(TagRanges[TagIdx].Name);
 	}
 
-	TArray<FName, TInlineAllocator<32>> SubBehaviorNames;
-	TArray<FAnimGenControlSchemaElement, TInlineAllocator<32>> SubBehaviorElements;
-
-	const int32 BehaviorNum = Behaviors.Num();
-
-	for (int32 BehaviorIdx = 0 ; BehaviorIdx < BehaviorNum; BehaviorIdx++)
-	{
-		if (Behaviors[BehaviorIdx].Behavior)
-		{
-			SubBehaviorNames.Add(Behaviors[BehaviorIdx].Name);
-			SubBehaviorElements.Add(UAnimGenControls::SpecifyEncodingControl(InControlSchema, Behaviors[BehaviorIdx].Behavior->SpecifyControl(InControlSchema)));
-		}
-	}
-
 	return UAnimGenControls::SpecifyStructControlFromArrayViews(InControlSchema,
 		{
 			TEXT("Behavior"),
 			TEXT("Tags"),
 		},
 		{
-			UAnimGenControls::SpecifyExclusiveUnionControlFromArrayViews(InControlSchema, SubBehaviorNames, SubBehaviorElements),
+			Behavior ? Behavior->SpecifyControl(InControlSchema) : UAnimGenControls::SpecifyNullControl(InControlSchema),
 			UAnimGenControls::SpecifyNamedInclusiveDiscreteControlFromArrayView(InControlSchema, TagNames),
 		});
 }
 
-void UAnimGenBehavior_Multi::MakeControlSets_Implementation(TArray<FAnimGenControlSet>& OutControlSets, FAnimGenControlObject& InControlObject, const UAnimDatabase* InDatabase, const FAnimDatabaseFrameRanges& InFrameRanges) const
+void UAnimGenBehavior_Tagged::MakeControlSets_Implementation(TArray<FAnimGenControlSet>& OutControlSets, FAnimGenControlObject& InControlObject, const UAnimDatabase* InDatabase, const FAnimDatabaseFrameRanges& InFrameRanges) const
 {
 	const int32 TagNum = TagRanges.Num();
 
@@ -731,6 +364,91 @@ void UAnimGenBehavior_Multi::MakeControlSets_Implementation(TArray<FAnimGenContr
 		}
 	}
 
+	TArray<FAnimGenControlSet> SubBehaviorControlSets;
+	if (Behavior)
+	{
+		Behavior->MakeControlSets(SubBehaviorControlSets, InControlObject, InDatabase, InFrameRanges);
+	}
+
+	for (const FAnimGenControlSet& ControlSet : SubBehaviorControlSets)
+	{
+		OutControlSets.Add(UAnimGenControls::MakeStructControlSetFromArrayViews(InControlObject,
+			{
+				TEXT("Behavior"),
+				TEXT("Tags")
+			},
+			{
+				ControlSet,
+				UAnimGenControls::MakeNamedInclusiveDiscreteControlSetFromFrameRangesArrayViews(InControlObject, UAnimGenControls::ControlSetFrameRanges(ControlSet), TagNames, TagFrameRanges),
+			}));
+	}
+}
+
+void UAnimGenBehavior_Tagged::DrawDebugControl_Implementation(
+	const FDebugDrawer& Drawer,
+	const FAnimGenControlObject& InControlObject,
+	const FAnimGenControlObjectElement InControlObjectElement,
+	const FAnimDatabasePoseState& InPoseState,
+	const UAnimDatabase* InDatabase,
+	const UAnimGenAutoEncoder* InAutoEncoder,
+	const int32 SequenceIdx,
+	const float SequenceTime,
+	const int32 RangeStart,
+	const int32 RangeLength,
+	const FLinearColor& IdentifierColor) const
+{
+	FAnimGenControlObjectElement BehaviorControl, TagsControl;
+	UAnimGenControls::GetStructControlElement(BehaviorControl, InControlObject, InControlObjectElement, TEXT("Behavior"));
+	UAnimGenControls::GetStructControlElement(TagsControl, InControlObject, InControlObjectElement, TEXT("Tags"));
+
+	UAnimGenControls::DrawDebugNamedInclusiveDiscreteControl(
+		Drawer, 
+		InControlObject, 
+		TagsControl, 
+		TEXT(""), 
+		UAnimDatabasePoseStateLibrary::PoseStateRootTransform(InPoseState).GetLocation() + TagDebugDrawOffset,
+		FRotator(0.0f, -90.0f, 0.0f), 
+		10.0f, 
+		IdentifierColor);
+
+	if (Behavior)
+	{
+		Behavior->DrawDebugControl(
+			Drawer,
+			InControlObject,
+			BehaviorControl,
+			InPoseState,
+			InDatabase,
+			InAutoEncoder,
+			SequenceIdx,
+			SequenceTime,
+			RangeStart,
+			RangeLength,
+			IdentifierColor);
+	}
+}
+
+FAnimGenControlSchemaElement UAnimGenBehavior_Multi::SpecifyControl_Implementation(FAnimGenControlSchema& InControlSchema) const
+{
+	TArray<FName, TInlineAllocator<32>> SubBehaviorNames;
+	TArray<FAnimGenControlSchemaElement, TInlineAllocator<32>> SubBehaviorElements;
+
+	const int32 BehaviorNum = Behaviors.Num();
+
+	for (int32 BehaviorIdx = 0 ; BehaviorIdx < BehaviorNum; BehaviorIdx++)
+	{
+		if (Behaviors[BehaviorIdx].Behavior)
+		{
+			SubBehaviorNames.Add(Behaviors[BehaviorIdx].Name);
+			SubBehaviorElements.Add(UAnimGenControls::SpecifyEncodingControl(InControlSchema, Behaviors[BehaviorIdx].Behavior->SpecifyControl(InControlSchema)));
+		}
+	}
+
+	return UAnimGenControls::SpecifyExclusiveUnionControlFromArrayViews(InControlSchema, SubBehaviorNames, SubBehaviorElements);
+}
+
+void UAnimGenBehavior_Multi::MakeControlSets_Implementation(TArray<FAnimGenControlSet>& OutControlSets, FAnimGenControlObject& InControlObject, const UAnimDatabase* InDatabase, const FAnimDatabaseFrameRanges& InFrameRanges) const
+{
 	const int32 BehaviorNum = Behaviors.Num();
 
 	for (int32 BehaviorIdx = 0; BehaviorIdx < BehaviorNum; BehaviorIdx++)
@@ -745,15 +463,7 @@ void UAnimGenBehavior_Multi::MakeControlSets_Implementation(TArray<FAnimGenContr
 
 			for (const FAnimGenControlSet& ControlSet : SubBehaviorControlSets)
 			{
-				OutControlSets.Add(UAnimGenControls::MakeStructControlSetFromArrayViews(InControlObject,
-					{
-						TEXT("Behavior"),
-						TEXT("Tags")
-					},
-					{
-						UAnimGenControls::MakeExclusiveUnionControlSet(InControlObject, Behaviors[BehaviorIdx].Name, UAnimGenControls::MakeEncodingControlSet(InControlObject, ControlSet)),
-						UAnimGenControls::MakeNamedInclusiveDiscreteControlSetFromFrameRangesArrayViews(InControlObject, UAnimGenControls::ControlSetFrameRanges(ControlSet), TagNames, TagFrameRanges),
-					}));
+				OutControlSets.Add(UAnimGenControls::MakeExclusiveUnionControlSet(InControlObject, Behaviors[BehaviorIdx].Name, UAnimGenControls::MakeEncodingControlSet(InControlObject, ControlSet)));
 			}
 		}
 	}
@@ -772,17 +482,11 @@ void UAnimGenBehavior_Multi::DrawDebugControl_Implementation(
 	const int32 RangeLength,
 	const FLinearColor& IdentifierColor) const
 {
-	FAnimGenControlObjectElement BehaviorControl, TagsControl;
-	UAnimGenControls::GetStructControlElement(BehaviorControl, InControlObject, InControlObjectElement, TEXT("Behavior"));
-	UAnimGenControls::GetStructControlElement(TagsControl, InControlObject, InControlObjectElement, TEXT("Tags"));
-
 	const FTransform RootTransform = UAnimDatabasePoseStateLibrary::PoseStateRootTransform(InPoseState);
 
-	UAnimGenControls::DrawDebugNamedInclusiveDiscreteControl(Drawer, InControlObject, TagsControl, TEXT(""), RootTransform.GetLocation() + FVector(0.0, 0.0, 260.0f), FRotator(0.0f, -90.0f, 0.0f), 10.0f, IdentifierColor);
-
 	FName BehaviorName = NAME_None;
 	FAnimGenControlObjectElement EncodingElement, BehaviorElement;
-	UAnimGenControls::GetExclusiveUnionControl(BehaviorName, EncodingElement, InControlObject, BehaviorControl);
+	UAnimGenControls::GetExclusiveUnionControl(BehaviorName, EncodingElement, InControlObject, InControlObjectElement);
 	UAnimGenControls::GetEncodingControl(BehaviorElement, InControlObject, EncodingElement);
 
 	FDrawDebugStringSettings Settings;
@@ -825,48 +529,21 @@ FAnimGenControlObjectElement UAnimGenBehaviorLibrary::MakeUncontrolledBehaviorCo
 	return UAnimGenControls::MakeNullControl(InControlObject);
 }
 
-FAnimGenControlObjectElement UAnimGenBehaviorLibrary::MakeIdleBehaviorControl(FAnimGenControlObject& InControlObject, const TArray<FName>& DesiredTags)
-{
-	return MakeIdleBehaviorControlFromArrayView(InControlObject, DesiredTags);
-}
-
-FAnimGenControlObjectElement UAnimGenBehaviorLibrary::MakeIdleBehaviorControlFromArrayView(FAnimGenControlObject& InControlObject, const TArrayView<const FName> DesiredTags)
+FAnimGenControlObjectElement UAnimGenBehaviorLibrary::MakeIdleBehaviorControl(FAnimGenControlObject& InControlObject)
 {
-	return UAnimGenControls::MakeNamedInclusiveDiscreteControlFromArrayView(InControlObject, DesiredTags);
-}
-
-FAnimGenControlObjectElement UAnimGenBehaviorLibrary::MakeTrajectoryFollowBehaviorControlFromArrayView(
-	FAnimGenControlObject& InControlObject,
-	const FTransformTrajectory& TransformTrajectory,
-	const FTransform& RootBoneTransform,
-	const TArrayView<const FName> DesiredTags,
-	const int32 InTrajectorySampleNum,
-	const float InTrajectoryFutureTime,
-	const float InTrajectoryPastTime,
-	const FVector InTrajectoryForwardVector)
-{
-	return UAnimGenControls::MakeStructControlFromArrayViews(InControlObject,
-		{
-			TEXT("Trajectory"),
-			TEXT("Tags")
-		},
-		{
-			UAnimGenControls::MakeTrajectoryControlFromTransformTrajectory(InControlObject, TransformTrajectory, RootBoneTransform, InTrajectorySampleNum, InTrajectoryFutureTime, InTrajectoryPastTime, InTrajectoryForwardVector),
-			UAnimGenControls::MakeNamedInclusiveDiscreteControlFromArrayView(InControlObject, DesiredTags),
-		});
+	return UAnimGenControls::MakeNullControl(InControlObject);
 }
 
 FAnimGenControlObjectElement UAnimGenBehaviorLibrary::MakeTrajectoryFollowBehaviorControl(
 	FAnimGenControlObject& InControlObject,
 	const FTransformTrajectory& TransformTrajectory,
 	const FTransform& RootBoneTransform,
-	const TArray<FName>& DesiredTags,
 	const int32 InTrajectorySampleNum,
 	const float InTrajectoryFutureTime,
 	const float InTrajectoryPastTime,
 	const FVector InTrajectoryForwardVector)
 {
-	return MakeTrajectoryFollowBehaviorControlFromArrayView(InControlObject, TransformTrajectory, RootBoneTransform, DesiredTags, InTrajectorySampleNum, InTrajectoryFutureTime, InTrajectoryPastTime, InTrajectoryForwardVector);
+	return UAnimGenControls::MakeTrajectoryControlFromTransformTrajectory(InControlObject, TransformTrajectory, RootBoneTransform, InTrajectorySampleNum, InTrajectoryFutureTime, InTrajectoryPastTime, InTrajectoryForwardVector);
 }
 
 FAnimGenControlObjectElement UAnimGenBehaviorLibrary::MakeMoveToTargetBehaviorControl(
@@ -874,32 +551,18 @@ FAnimGenControlObjectElement UAnimGenBehaviorLibrary::MakeMoveToTargetBehaviorCo
 	const FVector TargetLocation,
 	const FVector TargetDirection,
 	const bool bUseTargetDirection,
-	const FTransform& RootBoneTransform,
-	const TArray<FName>& DesiredTags)
-{
-	return MakeMoveToTargetBehaviorControlFromArrayView(InControlObject, TargetLocation, TargetDirection, bUseTargetDirection, RootBoneTransform, DesiredTags);
-}
-
-FAnimGenControlObjectElement UAnimGenBehaviorLibrary::MakeMoveToTargetBehaviorControlFromArrayView(
-	FAnimGenControlObject& InControlObject,
-	const FVector TargetLocation,
-	const FVector TargetDirection,
-	const bool bUseTargetDirection,
-	const FTransform& RootBoneTransform,
-	const TArrayView<const FName> DesiredTags)
+	const FTransform& RootBoneTransform)
 {
 	return UAnimGenControls::MakeStructControlFromArrayViews(InControlObject,
 		{
 			TEXT("Location"),
 			TEXT("Direction"),
-			TEXT("Tags")
 		},
 		{
 			UAnimGenControls::MakeLocationControl(InControlObject, TargetLocation, RootBoneTransform),
 			bUseTargetDirection ?
 				UAnimGenControls::MakeOptionalValidControl(InControlObject, UAnimGenControls::MakeDirectionControl(InControlObject, TargetDirection, RootBoneTransform)) :
 				UAnimGenControls::MakeOptionalNullControl(InControlObject),
-			UAnimGenControls::MakeNamedInclusiveDiscreteControlFromArrayView(InControlObject, DesiredTags),
 		});
 }
 
@@ -907,23 +570,6 @@ FAnimGenControlObjectElement UAnimGenBehaviorLibrary::MakeTrajectoryInteractionB
 	FAnimGenControlObject& InControlObject,
 	const FTransformTrajectory& TransformTrajectory,
 	const FTransform& RootBoneTransform,
-	const TArray<FName>& DesiredTags,
-	const FVector InteractionLocation,
-	const FVector InteractionDirection,
-	const bool bInteractionActive,
-	const int32 InTrajectorySampleNum,
-	const float InTrajectoryFutureTime,
-	const float InTrajectoryPastTime,
-	const FVector InTrajectoryForwardVector)
-{
-	return MakeTrajectoryInteractionBehaviorControlFromArrayView(InControlObject, TransformTrajectory, RootBoneTransform, DesiredTags, InteractionLocation, InteractionDirection, bInteractionActive, InTrajectorySampleNum, InTrajectoryFutureTime, InTrajectoryPastTime, InTrajectoryForwardVector);
-}
-
-FAnimGenControlObjectElement UAnimGenBehaviorLibrary::MakeTrajectoryInteractionBehaviorControlFromArrayView(
-	FAnimGenControlObject& InControlObject,
-	const FTransformTrajectory& TransformTrajectory,
-	const FTransform& RootBoneTransform,
-	const TArrayView<const FName> DesiredTags,
 	const FVector InteractionLocation,
 	const FVector InteractionDirection,
 	const bool bInteractionActive,
@@ -935,13 +581,11 @@ FAnimGenControlObjectElement UAnimGenBehaviorLibrary::MakeTrajectoryInteractionB
 	return UAnimGenControls::MakeStructControlFromArrayViews(InControlObject,
 		{
 			TEXT("Trajectory"),
-			TEXT("Tags"),
 			TEXT("Interaction")
 		},
 		{
 			UAnimGenControls::MakeTrajectoryControlFromTransformTrajectory(InControlObject, TransformTrajectory, RootBoneTransform, InTrajectorySampleNum, InTrajectoryFutureTime, InTrajectoryPastTime, InTrajectoryForwardVector),
-			UAnimGenControls::MakeNamedInclusiveDiscreteControlFromArrayView(InControlObject, DesiredTags),
-			bInteractionActive ? 
+			bInteractionActive ?
 				UAnimGenControls::MakeOptionalValidControl(InControlObject, UAnimGenControls::MakeStructControlFromArrayViews(InControlObject,
 					{
 						TEXT("Location"),
@@ -955,18 +599,16 @@ FAnimGenControlObjectElement UAnimGenBehaviorLibrary::MakeTrajectoryInteractionB
 		});
 }
 
-FAnimGenControlObjectElement UAnimGenBehaviorLibrary::MakeMultiBehaviorControl(
+FAnimGenControlObjectElement UAnimGenBehaviorLibrary::MakeTaggedBehaviorControl(
 	FAnimGenControlObject& InControlObject,
-	const FName Behavior,
 	const FAnimGenControlObjectElement& BehaviorControlElement,
 	const TArray<FName>& DesiredTags)
 {
-	return MakeMultiBehaviorControlFromArrayView(InControlObject, Behavior, BehaviorControlElement, DesiredTags);
+	return MakeTaggedBehaviorControlFromArrayView(InControlObject, BehaviorControlElement, DesiredTags);
 }
 
-FAnimGenControlObjectElement UAnimGenBehaviorLibrary::MakeMultiBehaviorControlFromArrayView(
+FAnimGenControlObjectElement UAnimGenBehaviorLibrary::MakeTaggedBehaviorControlFromArrayView(
 	FAnimGenControlObject& InControlObject,
-	const FName Behavior,
 	const FAnimGenControlObjectElement& BehaviorControlElement,
 	const TArrayView<const FName> DesiredTags)
 {
@@ -976,119 +618,19 @@ FAnimGenControlObjectElement UAnimGenBehaviorLibrary::MakeMultiBehaviorControlFr
 			TEXT("Tags"),
 		},
 		{
-			UAnimGenControls::MakeExclusiveUnionControl(InControlObject, Behavior, UAnimGenControls::MakeEncodingControl(InControlObject, BehaviorControlElement)),
+			BehaviorControlElement,
 			UAnimGenControls::MakeNamedInclusiveDiscreteControlFromArrayView(InControlObject, DesiredTags),
 		});
 }
 
-
-bool UAnimGenBehaviorLibrary::MakeRandomBool(int32& State)
-{
-	return UE::Learning::Random::SampleFloat((uint32&)State) <= 0.5f;
-}
-
-bool UAnimGenBehaviorLibrary::MakeRandomBoolWithProbability(int32& State, const float TrueProbability)
-{
-	return UE::Learning::Random::SampleFloat((uint32&)State) <= TrueProbability;
-}
-
-bool UAnimGenBehaviorLibrary::SwitchOnRandomBool(int32& State)
-{
-	return MakeRandomBool(State);
-}
-
-bool UAnimGenBehaviorLibrary::SwitchOnRandomBoolWithProbability(int32& State, const float TrueProbability)
-{
-	return MakeRandomBoolWithProbability(State, TrueProbability);
-}
-
-float UAnimGenBehaviorLibrary::MakeRandomFloat(int32& State)
-{
-	return UE::Learning::Random::SampleFloat((uint32&)State);
-}
-
-float UAnimGenBehaviorLibrary::MakeRandomFloatInRange(int32& State, const float Min, const float Max)
-{
-	return UE::Learning::Random::SampleFloat((uint32&)State) * (Max - Min) + Min;
-}
-
-float UAnimGenBehaviorLibrary::MakeRandomFloatChoice(int32& State, const TArray<float>& Values)
-{
-	return MakeRandomFloatChoiceFromArrayView(State, Values);
-}
-
-float UAnimGenBehaviorLibrary::MakeRandomFloatChoiceFromArrayView(int32& State, const TArrayView<const float> Values)
-{
-	if (Values.IsEmpty())
-	{
-		return 0.0f;
-	}
-
-	if (Values.Num() == 1)
-	{
-		return Values[0];
-	}
-
-	return Values[MakeRandomIntegerInRangeInclusive(State, 0, Values.Num() - 1)];
-}
-
-int32 UAnimGenBehaviorLibrary::MakeRandomIntegerInRangeInclusive(int32& State, const int32 Min, const int32 Max)
-{
-	const int32 ClampedMax = FMath::Max(Max, Min);
-	const uint32 RandomValue = UE::Learning::Random::SampleInt((uint32&)State);
-	return Min + RandomValue % ((ClampedMax - Min) + 1);
-}
-
-uint8 UAnimGenBehaviorLibrary::MakeRandomEnumValue(int32& State, const UEnum* Enum)
-{
-	if (Enum == nullptr || Enum->NumEnums() < 2) { return 0; }
-
-	const uint32 RandomValue = UE::Learning::Random::SampleInt((uint32&)State);
-	return Enum->GetValueByIndex(RandomValue % (Enum->NumEnums() - 1));
-}
-
-FVector UAnimGenBehaviorLibrary::MakeRandomVectorFromGaussianDistribution(int32& State, const FVector Mean, const FVector Std)
-{
-	return UE::Learning::Random::SampleVectorGaussian((uint32&)State, Mean, Std);
-}
-
-FVector UAnimGenBehaviorLibrary::MakeRandomDirection3D(int32& State, const FVector Mean, const FVector Std, const FVector Default)
-{
-	return UE::Learning::Random::SampleVectorGaussian((uint32&)State, Mean, Std).GetSafeNormal(UE_SMALL_NUMBER, Default);
-}
-
-FVector UAnimGenBehaviorLibrary::MakeRandomDirection2D(int32& State)
-{
-	return UE::Learning::Random::SampleVectorGaussian((uint32&)State, FVector::ZeroVector, FVector(1.0, 1.0, 0.0)).GetSafeNormal(UE_SMALL_NUMBER, FVector::ForwardVector);
-}
-
-FVector UAnimGenBehaviorLibrary::MakeRandomLocationFromGaussian2D(int32& State, const FVector Mean, const float Std)
-{
-	return MakeRandomVectorFromGaussianDistribution(State, Mean * FVector(1.0, 1.0, 0.0), FVector(Std, Std, 0.0));
-}
-
-void UAnimGenBehaviorLibrary::MakeRandomSequenceIndexAndTimeFromFrameRanges(
-	int32& OutSequenceIdx,
-	float& OutSequenceTime,
-	int32& State,
-	const FAnimDatabaseFrameRanges& FrameRanges,
-	const float FrameDeltaTime)
+FAnimGenControlObjectElement UAnimGenBehaviorLibrary::MakeMultiBehaviorControl(
+	FAnimGenControlObject& InControlObject,
+	const FName Behavior,
+	const FAnimGenControlObjectElement& BehaviorControlElement)
 {
-	int32 EntryIdx = INDEX_NONE;
-	int32 RangeIdx = INDEX_NONE;
-	int32 RangeFrame = INDEX_NONE;
-	UE::Learning::Random::SampleFrameInFrameRangeSet(
-		EntryIdx,
-		RangeIdx,
-		RangeFrame,
-		(uint32&)State,
-		*FrameRanges.FrameRangeSet);
-
-	OutSequenceIdx = FrameRanges.FrameRangeSet->GetEntrySequence(EntryIdx);
-	OutSequenceTime = (FrameRanges.FrameRangeSet->GetEntryRangeStart(EntryIdx, RangeIdx) + RangeFrame) * FrameDeltaTime;
+	return UAnimGenControls::MakeExclusiveUnionControl(InControlObject, Behavior, UAnimGenControls::MakeEncodingControl(InControlObject, BehaviorControlElement));
 }
 
-
 FVector UAnimGenBehaviorLibrary::ProjectOntoFloor(const FVector Location, const float FloorHeight)
 {
 	return FVector(Location.X, Location.Y, FloorHeight);
@@ -1388,315 +930,6 @@ bool UAnimGenBehaviorLibrary::LocationAndDirectionDifferenceBelowThreshold(
 	return bLocationClose && bDirectionClose;
 }
 
-void UAnimGenBehaviorLibrary::MakeTransformTrajectory(FTransformTrajectory& OutTransformTrajectory, const TArray<FVector>& Locations, const TArray<FRotator>& Rotations, const TArray<float>& Times)
-{
-	MakeTransformTrajectoryFromArrayViews(OutTransformTrajectory, Locations, Rotations, Times);
-}
-
-void UAnimGenBehaviorLibrary::MakeTransformTrajectoryFromArrayViews(FTransformTrajectory& OutTransformTrajectory, const TArrayView<const FVector> Locations, const TArrayView<const FRotator> Rotations, const TArrayView<const float> Times)
-{
-	const int32 SampleNum = FMath::Min(FMath::Min(Locations.Num(), Rotations.Num()), Times.Num());
-
-	OutTransformTrajectory.Samples.SetNumUninitialized(SampleNum);
-	for (int32 SampleIdx = 0; SampleIdx < SampleNum; SampleIdx++)
-	{
-		OutTransformTrajectory.Samples[SampleIdx].Position = Locations[SampleIdx];
-		OutTransformTrajectory.Samples[SampleIdx].Facing = Rotations[SampleIdx].Quaternion();
-		OutTransformTrajectory.Samples[SampleIdx].TimeInSeconds = Times[SampleIdx];
-	}
-}
-
-void UAnimGenBehaviorLibrary::MakeTransformTrajectoryFromPath(FTransformTrajectory& OutTransformTrajectory, const TArray<FVector>& Path, const float StartDistance, const float PathPointDistance, const float PathSpeed, const float FutureTime, const float PastTime, const float SampleInterval)
-{
-	return MakeTransformTrajectoryFromPathArrayView(OutTransformTrajectory, Path, StartDistance, PathPointDistance, PathSpeed, FutureTime, PastTime, SampleInterval);
-}
-
-void UAnimGenBehaviorLibrary::MakeTransformTrajectoryFromPathArrayView(FTransformTrajectory& OutTransformTrajectory, const TArrayView<const FVector> Path, const float StartDistance, const float PathPointDistance, const float PathSpeed, const float FutureTime, const float PastTime, const float SampleInterval)
-{
-	const int32 SampleNum = FMath::Max(FMath::RoundToInt((FutureTime + PastTime) / FMath::Max(SampleInterval, UE_SMALL_NUMBER)), 0);
-	OutTransformTrajectory.Samples.SetNumUninitialized(SampleNum);
-	for (int32 SampleIdx = 0; SampleIdx < SampleNum; SampleIdx++)
-	{
-		const float SamplePathTime = StartDistance + PathSpeed * (SampleIdx * SampleInterval - PastTime);
-
-		OutTransformTrajectory.Samples[SampleIdx].Position = PathLocationAtTimeFromArrayView(Path, SamplePathTime / FMath::Max(PathPointDistance, UE_SMALL_NUMBER));
-		OutTransformTrajectory.Samples[SampleIdx].Facing = PathRotationAtTimeFromArrayView(Path, SamplePathTime / FMath::Max(PathPointDistance, UE_SMALL_NUMBER)).Quaternion();
-		OutTransformTrajectory.Samples[SampleIdx].TimeInSeconds = SampleIdx * SampleInterval - PastTime;
-	}
-}
-
-void UAnimGenBehaviorLibrary::MakeTransformTrajectoryFromTrajectory(FTransformTrajectory& OutTransformTrajectory, const FAnimGenTrajectory& Trajectory, const float FutureTime, const float PastTime)
-{
-	const int32 SampleNum = FMath::Min(Trajectory.Locations.Num(), Trajectory.Rotations.Num());
-
-	OutTransformTrajectory.Samples.SetNumUninitialized(SampleNum);
-	for (int32 SampleIdx = 0; SampleIdx < SampleNum; SampleIdx++)
-	{
-		OutTransformTrajectory.Samples[SampleIdx].Position = Trajectory.Locations[SampleIdx];
-		OutTransformTrajectory.Samples[SampleIdx].Facing = Trajectory.Rotations[SampleIdx].Quaternion();
-		OutTransformTrajectory.Samples[SampleIdx].TimeInSeconds = (FutureTime + PastTime) * (((float)SampleIdx) / FMath::Max(SampleNum - 1, 1)) - PastTime;
-	}
-}
-
-void UAnimGenBehaviorLibrary::ProjectTransformTrajectoryOntoFloor(FTransformTrajectory& InOutTransformTrajectory, const float FloorHeight)
-{
-	const int32 SampleNum = InOutTransformTrajectory.Samples.Num();
-	for (int32 SampleIdx = 0; SampleIdx < SampleNum; SampleIdx++)
-	{
-		InOutTransformTrajectory.Samples[SampleIdx].Position.Z = FloorHeight;
-	}
-}
-
-int32 UAnimGenBehaviorLibrary::GetTrajectorySampleNum(const FAnimGenTrajectory& Trajectory)
-{
-	return Trajectory.Locations.Num();
-}
-
-const TArray<FVector>& UAnimGenBehaviorLibrary::GetTrajectoryLocations(const FAnimGenTrajectory& Trajectory)
-{
-	return Trajectory.Locations;
-}
-
-const TArray<FVector>& UAnimGenBehaviorLibrary::GetTrajectoryDirections(const FAnimGenTrajectory& Trajectory)
-{
-	return Trajectory.Directions;
-}
-
-void UAnimGenBehaviorLibrary::ResetTrajectory(
-	FAnimGenTrajectory& Trajectory,
-	const int32 SampleNum,
-	const FVector InitialLocation,
-	const FRotator InitialRotation,
-	const FVector ForwardVector)
-{
-	if (SampleNum < 0)
-	{
-		// TODO: Warning
-		return;
-	}
-
-	Trajectory.Locations.Init(InitialLocation, SampleNum);
-	Trajectory.Velocities.Init(FVector::ZeroVector, SampleNum);
-	Trajectory.Accelerations.Init(FVector::ZeroVector, SampleNum);
-	Trajectory.Rotations.Init(InitialRotation, SampleNum);
-	Trajectory.AngularVelocities.Init(FVector::ZeroVector, SampleNum);
-	Trajectory.Directions.Init(InitialRotation.RotateVector(ForwardVector), SampleNum);
-}
-
-void UAnimGenBehaviorLibrary::UpdateAndPredictFutureTrajectory(
-	FAnimGenTrajectory& Trajectory,
-	const FVector DesiredVelocity,
-	const FRotator DesiredRotation,
-	const float DeltaTime,
-	const float FutureTime,
-	const float MaximumDesiredVelocity,
-	const float MaximumAcceleration,
-	const float MaximumAngularVelocity,
-	const FVector ForwardVector)
-{
-	const int32 SampleNum = Trajectory.Locations.Num();
-
-	if (SampleNum < 1)
-	{
-		// TODO: Warning
-		return;
-	}
-
-	const float VelocityHalfLife = UE::AnimDatabase::Math::DampingToHalfLife(
-		2.0f * ((MaximumAcceleration / FMath::Max(MaximumDesiredVelocity, UE_KINDA_SMALL_NUMBER)) * UE_EULERS_NUMBER));
-
-	UE::AnimDatabase::Math::CriticalSpringUpdatePositionFromVelocity(
-		Trajectory.Locations[0],
-		Trajectory.Velocities[0],
-		Trajectory.Accelerations[0],
-		DesiredVelocity,
-		VelocityHalfLife,
-		DeltaTime);
-
-	for (int32 SampleIdx = 1; SampleIdx < SampleNum; SampleIdx++)
-	{
-		Trajectory.Locations[SampleIdx] = Trajectory.Locations[0];
-		Trajectory.Velocities[SampleIdx] = Trajectory.Velocities[0];
-		Trajectory.Accelerations[SampleIdx] = Trajectory.Accelerations[0];
-
-		UE::AnimDatabase::Math::CriticalSpringUpdatePositionFromVelocity(
-			Trajectory.Locations[SampleIdx],
-			Trajectory.Velocities[SampleIdx],
-			Trajectory.Accelerations[SampleIdx],
-			DesiredVelocity,
-			VelocityHalfLife,
-			(FutureTime / (SampleNum - 1)) * SampleIdx);
-	}
-
-	const float RotationHalfLife = UE::AnimDatabase::Math::DampingToHalfLife(
-		2.0f * ((MaximumAngularVelocity / 180.0f) * UE_EULERS_NUMBER));
-
-	UE::AnimDatabase::Math::CriticalSpringUpdate(
-		Trajectory.Rotations[0],
-		Trajectory.AngularVelocities[0],
-		DesiredRotation,
-		RotationHalfLife,
-		DeltaTime);
-
-	for (int32 SampleIdx = 1; SampleIdx < SampleNum; SampleIdx++)
-	{
-		Trajectory.Rotations[SampleIdx] = Trajectory.Rotations[0];
-		Trajectory.AngularVelocities[SampleIdx] = Trajectory.AngularVelocities[0];
-
-		UE::AnimDatabase::Math::CriticalSpringUpdate(
-			Trajectory.Rotations[SampleIdx],
-			Trajectory.AngularVelocities[SampleIdx],
-			DesiredRotation,
-			RotationHalfLife,
-			(FutureTime / (SampleNum - 1)) * SampleIdx);
-	}
-
-	for (int32 SampleIdx = 0; SampleIdx < SampleNum; SampleIdx++)
-	{
-		Trajectory.Directions[SampleIdx] = Trajectory.Rotations[SampleIdx].RotateVector(ForwardVector);
-	}
-}
-
-
-void UAnimGenBehaviorLibrary::TeleportTrajectory(
-	FAnimGenTrajectory& Trajectory,
-	const FVector NewLocation,
-	const FRotator NewRotation)
-{
-	if (Trajectory.Locations.Num() < 1)
-	{
-		// TODO: Warning
-		return;
-	}
-
-	const FVector SourceLocation = Trajectory.Locations[0];
-	const FQuat SourceRotation = Trajectory.Rotations[0].Quaternion();
-	const FVector DestinationLocation = NewLocation;
-	const FQuat DestinationRotation = NewRotation.Quaternion();
-
-	const int32 SampleNum = Trajectory.Locations.Num();
-
-	for (int32 SampleIdx = 0; SampleIdx < SampleNum; SampleIdx++)
-	{
-		Trajectory.Locations[SampleIdx] = DestinationRotation.RotateVector(SourceRotation.UnrotateVector(Trajectory.Locations[SampleIdx] - SourceLocation)) + DestinationLocation;
-		Trajectory.Velocities[SampleIdx] = DestinationRotation.RotateVector(SourceRotation.UnrotateVector(Trajectory.Velocities[SampleIdx]));
-		Trajectory.Accelerations[SampleIdx] = DestinationRotation.RotateVector(SourceRotation.UnrotateVector(Trajectory.Accelerations[SampleIdx]));
-		Trajectory.Rotations[SampleIdx] = (DestinationRotation * SourceRotation.Inverse() * Trajectory.Rotations[SampleIdx].Quaternion()).Rotator();
-		Trajectory.AngularVelocities[SampleIdx] = DestinationRotation.RotateVector(SourceRotation.UnrotateVector(Trajectory.AngularVelocities[SampleIdx]));
-		Trajectory.Directions[SampleIdx] = DestinationRotation.RotateVector(SourceRotation.UnrotateVector(Trajectory.Directions[SampleIdx]));
-	}
-}
-
-FVector UAnimGenBehaviorLibrary::PathVelocityAtTime(const TArray<FVector>& Points, const float Time)
-{
-	return PathVelocityAtTimeFromArrayView(Points, Time);
-}
-
-FVector UAnimGenBehaviorLibrary::PathVelocityAtTimeFromArrayView(const TArrayView<const FVector> Points, const float Time)
-{
-	if (Points.IsEmpty())
-	{
-		return FVector::ZeroVector;
-	}
-
-	FVector OutLocation = FVector::ZeroVector;
-	FVector3f OutVelocity = FVector3f::ZeroVector;
-	UE::AnimDatabase::Math::CubicMonoSampleLocation(OutLocation, OutVelocity, Points, Time, 1.0f);
-	return ((FVector)OutVelocity);
-}
-
-FVector UAnimGenBehaviorLibrary::PathLocationAtTime(const TArray<FVector>& Points, const float Time)
-{
-	return PathLocationAtTimeFromArrayView(Points, Time);
-}
-
-FVector UAnimGenBehaviorLibrary::PathLocationAtTimeFromArrayView(const TArrayView<const FVector> Points, const float Time)
-{
-	if (Points.IsEmpty())
-	{
-		return FVector::ZeroVector;
-	}
-
-	FVector OutLocation = FVector::ZeroVector;
-	UE::AnimDatabase::Math::CubicMonoSampleLocation(OutLocation, Points, Time, 1.0f);
-	return OutLocation;
-}
-
-FVector UAnimGenBehaviorLibrary::PathDirectionAtTime(const TArray<FVector>& Points, const float Time)
-{
-	return PathDirectionAtTimeFromArrayView(Points, Time);
-}
-
-FVector UAnimGenBehaviorLibrary::PathDirectionAtTimeFromArrayView(const TArrayView<const FVector> Points, const float Time)
-{
-	if (Points.IsEmpty())
-	{
-		return FVector::ForwardVector;
-	}
-
-	FVector OutLocation = FVector::ZeroVector;
-	FVector3f OutVelocity = FVector3f::ZeroVector;
-	UE::AnimDatabase::Math::CubicMonoSampleLocation(OutLocation, OutVelocity, Points, Time, 1.0f);
-	return ((FVector)OutVelocity).GetSafeNormal(UE_SMALL_NUMBER, FVector::ForwardVector);
-}
-
-FRotator UAnimGenBehaviorLibrary::PathRotationAtTime(const TArray<FVector>& Points, const float Time)
-{
-	return PathRotationAtTimeFromArrayView(Points, Time);
-}
-
-FRotator UAnimGenBehaviorLibrary::PathRotationAtTimeFromArrayView(const TArrayView<const FVector> Points, const float Time)
-{
-	if (Points.IsEmpty())
-	{
-		return FRotator::ZeroRotator;
-	}
-
-	FVector OutLocation = FVector::ZeroVector;
-	FVector3f OutVelocity = FVector3f::ZeroVector;
-	UE::AnimDatabase::Math::CubicMonoSampleLocation(OutLocation, OutVelocity, Points, Time, 1.0f);
-	const FVector OutDirection = ((FVector)OutVelocity).GetSafeNormal(UE_SMALL_NUMBER, FVector::ForwardVector);
-	return FVector(OutDirection.Y, -OutDirection.X, 0.0f).ToOrientationRotator();
-}
-
-bool UAnimGenBehaviorLibrary::PathAddPoint(TArray<FVector>& Points, const FVector NewPoint, const float MinDistance)
-{
-	if (Points.IsEmpty())
-	{
-		Points.Add(NewPoint);
-		return true;
-	}
-	else
-	{
-		if ((Points.Last() - NewPoint).Length() > MinDistance)
-		{
-			Points.Add(NewPoint);
-			return true;
-		}
-		else
-		{
-			return false;
-		}
-	}
-}
-
-float UAnimGenBehaviorLibrary::PathTotalDistance(const TArray<FVector>& Points)
-{
-	return PathTotalDistanceFromArrayView(Points);
-}
-
-float UAnimGenBehaviorLibrary::PathTotalDistanceFromArrayView(const TArrayView<const FVector> Points)
-{
-	const int32 PointNum = Points.Num();
-
-	float Distance = 0.0f;
-	for (int32 PointIdx = 0; PointIdx < PointNum - 1; PointIdx++)
-	{
-		Distance += FVector::Distance(Points[PointIdx + 0], Points[PointIdx + 1]);
-	}
-	return Distance;
-}
-
 FVector UAnimGenBehaviorLibrary::ApplySimpleAvoidanceToVelocity(
 	const FVector Velocity,
 	const TArray<AActor*>& Actors,
@@ -1790,383 +1023,3 @@ AActor* UAnimGenBehaviorLibrary::FindLookAtActorArrayView(
 
 	return nullptr;
 }
-
-#if WITH_EDITOR
-
-FAnimGenControlObjectElement UAnimGenBehaviorPreview_Uncontrolled::MakeControl_Implementation(FAnimGenControlObject& InControlObject, const FAnimDatabasePoseState& InPoseState) const
-{
-	return UAnimGenBehaviorLibrary::MakeUncontrolledBehaviorControl(InControlObject);
-}
-
-FAnimGenControlObjectElement UAnimGenBehaviorPreview_Idle::MakeControl_Implementation(FAnimGenControlObject& InControlObject, const FAnimDatabasePoseState& InPoseState) const
-{
-	TArray<FName, TInlineAllocator<1>> DesiredTags;
-	DesiredTags.Reserve(1);
-	if (DesiredStyle != NAME_None) { DesiredTags.Add(DesiredStyle); }
-
-	return UAnimGenBehaviorLibrary::MakeIdleBehaviorControlFromArrayView(InControlObject, DesiredTags);
-}
-
-void UAnimGenBehaviorPreview_Idle::OnPreviewInputKeyPressed_Implementation(const FKey Key)
-{
-	if (Key == EKeys::Gamepad_DPad_Left)
-	{
-		StyleIndex = FMath::WrapExclusive((StyleIndex - 1), 0, StyleTags.Num() + 1);
-		DesiredStyle = StyleIndex == 0 ? NAME_None : StyleTags[StyleIndex - 1];
-	}
-
-	if (Key == EKeys::Gamepad_DPad_Right)
-	{
-		StyleIndex = FMath::WrapExclusive((StyleIndex + 1), 0, StyleTags.Num() + 1);
-		DesiredStyle = StyleIndex == 0 ? NAME_None : StyleTags[StyleIndex - 1];
-	}
-}
-
-void UAnimGenBehaviorPreview_Idle::OnPreviewDrawDebug_Implementation(
-	const FDebugDrawer& Drawer,
-	const FAnimGenControlObject& InControlObject,
-	const FAnimGenControlObjectElement InControlObjectElement,
-	const FAnimDatabasePoseState& InPoseState,
-	const FLinearColor& IdentifierColor)
-{
-	const FTransform RootTransform = UAnimDatabasePoseStateLibrary::PoseStateRootTransform(InPoseState);
-
-	UAnimGenControls::DrawDebugNamedInclusiveDiscreteControl(Drawer, InControlObject, InControlObjectElement, TEXT(""), RootTransform.GetLocation() + FVector(0.0, 0.0, 220.0f), FRotator(0.0f, -90.0f, 0.0f), 10.0f, IdentifierColor);
-}
-
-FAnimGenControlObjectElement UAnimGenBehaviorPreview_TrajectoryFollow::MakeControl_Implementation(FAnimGenControlObject& InControlObject, const FAnimDatabasePoseState& InPoseState) const
-{
-	TArray<FName, TInlineAllocator<2>> DesiredTags;
-	DesiredTags.Reserve(2);
-	if (DesiredGait != NAME_None) { DesiredTags.Add(DesiredGait); }
-	if (DesiredStyle != NAME_None) { DesiredTags.Add(DesiredStyle); }
-
-	return UAnimGenBehaviorLibrary::MakeTrajectoryFollowBehaviorControlFromArrayView(
-		InControlObject,
-		TransformTrajectory,
-		UAnimDatabasePoseStateLibrary::PoseStateRootTransform(InPoseState),
-		DesiredTags,
-		TrajectorySampleNum,
-		TrajectoryFutureTime,
-		TrajectoryPastTime,
-		TrajectoryForwardVector);
-}
-
-void UAnimGenBehaviorPreview_TrajectoryFollow::OnPreviewReset_Implementation(const FAnimDatabasePoseState& InPoseState)
-{
-	UAnimGenBehaviorLibrary::ResetTrajectory(Trajectory, TrajectorySampleNum, FVector::ZeroVector, FRotator::ZeroRotator, TrajectoryForwardVector);
-	bFollowPath = false;
-	bDrawPath = false;
-	FollowPath.Reset();
-	DrawPath.Reset();
-}
-
-void UAnimGenBehaviorPreview_TrajectoryFollow::OnPreviewUpdate_Implementation(const FAnimDatabasePoseState& InPoseState, const float DeltaTime)
-{
-	FVector RightStickAxis = FVector::ZeroVector;
-	FVector RightStickDirection = FVector::ZeroVector;
-	float RightStickLength = 0.0f;
-	bool bRightStickInDeadzone = false;
-	UAnimGenBehaviorLibrary::ApplyStickDeadzoneAndSensitivity(
-		RightStickAxis,
-		RightStickDirection,
-		RightStickLength,
-		bRightStickInDeadzone,
-		GetPreviewRightStickAxis());
-
-	SpringMath::CriticalSpringDamper(CameraYawVelocity, CameraYawAcceleration, CameraRotationSpeed * (float)RightStickAxis.Y, CameraInputSmoothingTime, DeltaTime);
-	SpringMath::CriticalSpringDamper(CameraPitchVelocity, CameraPitchAcceleration, CameraRotationSpeed * (float)RightStickAxis.X, CameraInputSmoothingTime, DeltaTime);
-	SpringMath::CriticalSpringDamper(CameraDistanceVelocity, CameraDistanceAcceleration, CameraZoomSpeed * CameraDistanceDelta, CameraInputSmoothingTime, DeltaTime);
-
-	UpdatePreviewCamera(
-		CameraYawVelocity,
-		CameraPitchVelocity,
-		CameraDistanceVelocity,
-		DeltaTime,
-		bCameraInvertPitch);
-
-	if (bFollowPath && UAnimGenBehaviorLibrary::IsStickNotInDeadzone(GetPreviewLeftStickAxis()))
-	{
-		bFollowPath = false;
-	}
-
-	if (bFollowPath)
-	{
-		UAnimGenBehaviorLibrary::MakeTransformTrajectoryFromPathArrayView(
-			TransformTrajectory,
-			FollowPath,
-			PathDistance,
-			PathPointDistance,
-			PathFollowSpeed,
-			TrajectoryFutureTime);
-
-		if (bPathFacingCamera)
-		{
-			const int32 SampleNum = TransformTrajectory.Samples.Num();
-			for (int32 SampleIdx = 0; SampleIdx < SampleNum; SampleIdx++)
-			{
-				TransformTrajectory.Samples[SampleIdx].Facing = GetPreviewCameraYawRotation().RotateVector(-FVector::RightVector).ToOrientationQuat();
-			}
-		}
-
-		PathDistance += DeltaTime * PathFollowSpeed;
-	}
-	else
-	{
-		SpringMath::CriticalSpringDamper(DesiredMinimumSpeed, DesiredMinimumSpeedVelocity, bWalkActive ? MinimumWalkSpeed : MinimumRunSpeed, SpeedSmoothingTime, DeltaTime);
-		SpringMath::CriticalSpringDamper(DesiredMaximumSpeed, DesiredMaximumSpeedVelocity, bWalkActive ? MaximumWalkSpeed : MaximumRunSpeed, SpeedSmoothingTime, DeltaTime);
-
-		UAnimGenBehaviorLibrary::UpdateDesiredVelocityAndRotationFromAxisAndCameraTrigger(
-			DesiredVelocity,
-			DesiredRotation,
-			GetPreviewCameraYawRotation(),
-			GetPreviewLeftStickAxis(),
-			GetPreviewLeftTriggerAxisWithSensitivity(),
-			DesiredMinimumSpeed,
-			DesiredMaximumSpeed);
-
-		UAnimGenBehaviorLibrary::TeleportTrajectory(
-			Trajectory,
-			UAnimDatabasePoseStateLibrary::PoseStateRootLocation(InPoseState),
-			UAnimDatabasePoseStateLibrary::PoseStateRootRotation(InPoseState));
-
-		UAnimGenBehaviorLibrary::UpdateAndPredictFutureTrajectory(
-			Trajectory,
-			DesiredVelocity,
-			DesiredRotation,
-			DeltaTime,
-			TrajectoryFutureTime,
-			DesiredMaximumSpeed,
-			MaximumAcceleration,
-			MaximumAngularVelocity,
-			TrajectoryForwardVector);
-
-		UAnimGenBehaviorLibrary::MakeTransformTrajectoryFromTrajectory(TransformTrajectory, Trajectory, TrajectoryFutureTime);
-
-		if (DesiredVelocity.IsNearlyZero())
-		{
-			DesiredGait = StandTag;
-		}
-		else
-		{
-			DesiredGait = bWalkActive ? WalkTag : RunTag;
-		}
-	}
-}
-
-void UAnimGenBehaviorPreview_TrajectoryFollow::OnPreviewInputKeyPressed_Implementation(const FKey Key)
-{
-	if (Key == EKeys::Gamepad_RightShoulder)
-	{
-		bWalkActive = !bWalkActive;
-	}
-
-	if (Key == EKeys::Gamepad_DPad_Left)
-	{
-		StyleIndex = FMath::WrapExclusive((StyleIndex - 1), 0, StyleTags.Num() + 1);
-		DesiredStyle = StyleIndex == 0 ? NAME_None : StyleTags[StyleIndex - 1];
-	}
-
-	if (Key == EKeys::Gamepad_DPad_Right)
-	{
-		StyleIndex = FMath::WrapExclusive((StyleIndex + 1), 0, StyleTags.Num() + 1);
-		DesiredStyle = StyleIndex == 0 ? NAME_None : StyleTags[StyleIndex - 1];
-	}
-
-	if (Key == EKeys::Gamepad_DPad_Up)
-	{
-		CameraDistanceDelta = +1.0f;
-	}
-
-	if (Key == EKeys::Gamepad_DPad_Down)
-	{
-		CameraDistanceDelta = -1.0f;
-	}
-}
-
-void UAnimGenBehaviorPreview_TrajectoryFollow::OnPreviewInputKeyReleased_Implementation(const FKey Key)
-{
-	if (Key == EKeys::Gamepad_DPad_Up)
-	{
-		CameraDistanceDelta = 0.0f;
-	}
-
-	if (Key == EKeys::Gamepad_DPad_Down)
-	{
-		CameraDistanceDelta = 0.0f;
-	}
-}
-
-void UAnimGenBehaviorPreview_TrajectoryFollow::OnPreviewDrawDebug_Implementation(
-	const FDebugDrawer& Drawer,
-	const FAnimGenControlObject& InControlObject,
-	const FAnimGenControlObjectElement InControlObjectElement,
-	const FAnimDatabasePoseState& InPoseState,
-	const FLinearColor& IdentifierColor)
-{
-	FAnimGenControlObjectElement TrajectoryControl, TagsControl;
-	UAnimGenControls::GetStructControlElement(TrajectoryControl, InControlObject, InControlObjectElement, TEXT("Trajectory"));
-	UAnimGenControls::GetStructControlElement(TagsControl, InControlObject, InControlObjectElement, TEXT("Tags"));
-
-	const FTransform RootTransform = UAnimDatabasePoseStateLibrary::PoseStateRootTransform(InPoseState);
-
-	UAnimGenControls::DrawDebugTrajectoryControl(Drawer, InControlObject, TrajectoryControl, RootTransform, IdentifierColor, 25.0f, 5.0f, 1.0f, 1.0f);
-	UAnimGenControls::DrawDebugNamedInclusiveDiscreteControl(Drawer, InControlObject, TagsControl, TEXT(""), RootTransform.GetLocation() + FVector(0.0, 0.0, 220.0f), FRotator(0.0f, -90.0f, 0.0f), 10.0f, IdentifierColor);
-
-	if (bDrawPath)
-	{
-		UDrawDebugLibrary::DrawDebugCatmullRomSplineArrayView(
-			Drawer,
-			DrawPath,
-			UDrawDebugLibrary::MakeDrawDebugLineStyleFromColorAndThickness(IdentifierColor, 2.0f), true, true, 8);
-
-		UDrawDebugLibrary::DrawDebugLocationsArrayView(
-			Drawer, 
-			DrawPath, 
-			UDrawDebugLibrary::MakeDrawDebugLineStyleFromColorAndThickness(IdentifierColor, 2.0f), true, 5.0f);
-	}
-	else if (bFollowPath)
-	{
-		UDrawDebugLibrary::DrawDebugCatmullRomSplineArrayView(
-			Drawer, 
-			FollowPath, 
-			UDrawDebugLibrary::MakeDrawDebugLineStyleFromColorAndThickness(IdentifierColor, 2.0f), true, true, 8);
-
-		UDrawDebugLibrary::DrawDebugLocationsArrayView(
-			Drawer,
-			FollowPath,
-			UDrawDebugLibrary::MakeDrawDebugLineStyleFromColorAndThickness(IdentifierColor, 2.0f), true, 5.0f);
-	}
-}
-
-void UAnimGenBehaviorPreview_TrajectoryFollow::OnPreviewInputLeftMousePressed_Implementation(const FVector MouseGroundLocation)
-{
-	bDrawPath = true;
-	DrawPath.Reset();
-	DrawPath.Add(MouseGroundLocation);
-}
-
-void UAnimGenBehaviorPreview_TrajectoryFollow::OnPreviewInputLeftMouseReleased_Implementation(const FVector MouseGroundLocation)
-{
-	bFollowPath = true;
-	bDrawPath = false;
-	FollowPath = DrawPath;
-	PathDistance = -PathFollowSpeed;
-}
-
-void UAnimGenBehaviorPreview_TrajectoryFollow::OnPreviewInputMouseMoved_Implementation(const FVector MouseGroundLocation)
-{
-	if (bDrawPath)
-	{
-		UAnimGenBehaviorLibrary::PathAddPoint(DrawPath, MouseGroundLocation, PathPointDistance);
-	}
-}
-
-FAnimGenControlObjectElement UAnimGenBehaviorPreview_MoveToTarget::MakeControl_Implementation(FAnimGenControlObject& InControlObject, const FAnimDatabasePoseState& InPoseState) const
-{
-	TArray<FName, TInlineAllocator<1>> DesiredTags;
-	DesiredTags.Reserve(1);
-	if (DesiredStyle != NAME_None) { DesiredTags.Add(DesiredStyle); }
-
-	return UAnimGenBehaviorLibrary::MakeMoveToTargetBehaviorControlFromArrayView(
-		InControlObject,
-		TargetLocation,
-		TargetDirection,
-		bUseTargetDirection,
-		UAnimDatabasePoseStateLibrary::PoseStateRootTransform(InPoseState),
-		DesiredTags);
-}
-
-void UAnimGenBehaviorPreview_MoveToTarget::OnPreviewInputKeyPressed_Implementation(const FKey Key)
-{
-	if (Key == EKeys::Gamepad_DPad_Left)
-	{
-		StyleIndex = FMath::WrapExclusive((StyleIndex - 1), 0, StyleTags.Num() + 1);
-		DesiredStyle = StyleIndex == 0 ? NAME_None : StyleTags[StyleIndex - 1];
-	}
-
-	if (Key == EKeys::Gamepad_DPad_Right)
-	{
-		StyleIndex = FMath::WrapExclusive((StyleIndex + 1), 0, StyleTags.Num() + 1);
-		DesiredStyle = StyleIndex == 0 ? NAME_None : StyleTags[StyleIndex - 1];
-	}
-}
-
-void UAnimGenBehaviorPreview_MoveToTarget::OnPreviewDrawDebug_Implementation(
-	const FDebugDrawer& Drawer,
-	const FAnimGenControlObject& InControlObject,
-	const FAnimGenControlObjectElement InControlObjectElement,
-	const FAnimDatabasePoseState& InPoseState,
-	const FLinearColor& IdentifierColor)
-{
-	if (bPlacingActive)
-	{
-		UDrawDebugLibrary::DrawDebugLocation(
-			Drawer,
-			ActiveLocation,
-			UDrawDebugLibrary::MakeDrawDebugLineStyleFromColorAndThickness(FLinearColor::White.CopyWithNewOpacity(IdentifierColor.A), 1.0f),
-			true,
-			10.0f);
-
-		if (bUseActiveDirection)
-		{
-			UDrawDebugLibrary::DrawDebugDirection(
-				Drawer,
-				ActiveLocation,
-				ActiveDirection,
-				UDrawDebugLibrary::MakeDrawDebugLineStyleFromColorAndThickness(FLinearColor::White.CopyWithNewOpacity(IdentifierColor.A), 1.0f),
-				true,
-				25.0f,
-				1.0f);
-		}
-	}
-
-	UDrawDebugLibrary::DrawDebugLocation(
-		Drawer,
-		TargetLocation,
-		UDrawDebugLibrary::MakeDrawDebugLineStyleFromColorAndThickness(IdentifierColor, 1.0f),
-		true,
-		10.0f);
-
-	if (bUseTargetDirection)
-	{
-		UDrawDebugLibrary::DrawDebugDirection(
-			Drawer,
-			TargetLocation,
-			TargetDirection,
-			UDrawDebugLibrary::MakeDrawDebugLineStyleFromColorAndThickness(IdentifierColor, 1.0f),
-			true,
-			25.0f,
-			1.0f);
-	}
-}
-
-void UAnimGenBehaviorPreview_MoveToTarget::OnPreviewInputLeftMousePressed_Implementation(const FVector MouseGroundLocation)
-{
-	ActiveLocation = MouseGroundLocation;
-	ActiveDirection = FVector::ForwardVector;
-	bUseActiveDirection = false;
-	bPlacingActive = true;
-}
-
-void UAnimGenBehaviorPreview_MoveToTarget::OnPreviewInputLeftMouseReleased_Implementation(const FVector MouseGroundLocation)
-{
-	TargetLocation = ActiveLocation;
-	TargetDirection = ActiveDirection;
-	bUseTargetDirection = bUseActiveDirection;
-	bPlacingActive = false;
-}
-
-void UAnimGenBehaviorPreview_MoveToTarget::OnPreviewInputMouseMoved_Implementation(const FVector MouseGroundLocation)
-{
-	if (bPlacingActive)
-	{
-		if (FVector::Distance(MouseGroundLocation, ActiveLocation) > 10.0f)
-		{
-			ActiveDirection = (MouseGroundLocation - ActiveLocation).GetSafeNormal();
-			bUseActiveDirection = true;
-		}
-	}
-}
-
-#endif
\ No newline at end of file
diff --git a/Engine/Plugins/Experimental/Animation/AnimGen/Source/AnimGen/Private/AnimGenController.cpp b/Engine/Plugins/Experimental/Animation/AnimGen/Source/AnimGen/Private/AnimGenController.cpp
--- a/Engine/Plugins/Experimental/Animation/AnimGen/Source/AnimGen/Private/AnimGenController.cpp
+++ b/Engine/Plugins/Experimental/Animation/AnimGen/Source/AnimGen/Private/AnimGenController.cpp
@@ -245,26 +245,12 @@ bool UAnimGenControllerTrainingSettings::Update()
 	return false;
 }
 
-void UAnimGenControllerTestingSettings::LoadDatabaseAsync()
-{
-	ResetDatabase.LoadAsync(FLoadSoftObjectPathAsyncDelegate());
-}
-
-void UAnimGenControllerTestingSettings::PostEditChangeProperty(struct FPropertyChangedEvent& Event)
-{
-	if (Event.GetMemberPropertyName() == GET_MEMBER_NAME_CHECKED(UAnimGenControllerTestingSettings, ResetDatabase))
-	{
-		LoadDatabaseAsync();
-	}
-}
-
 #endif
 
 UAnimGenController::UAnimGenController(const FObjectInitializer& ObjectInitializer)
 {
 #if WITH_EDITOR
 	TrainingSettings = ObjectInitializer.CreateEditorOnlyDefaultSubobject<UAnimGenControllerTrainingSettings>(this, TEXT("TrainingSettings"));
-	TestingSettings = ObjectInitializer.CreateEditorOnlyDefaultSubobject<UAnimGenControllerTestingSettings>(this, TEXT("TestingSettings"));
 	ViewportSettings = ObjectInitializer.CreateEditorOnlyDefaultSubobject<UAnimGenControllerViewportSettings>(this, TEXT("ViewportSettings"));
 #endif
 }
@@ -300,8 +286,6 @@ void UAnimGenController::Invalidate()
 	TrainedSchemaCompatibilityHash = 0;
 	ControlEncoderSize = 0;
 	DenoiserSize = 0;
-	ControlEncoderInferenceTime = 0;
-	DenoiserInferenceTime = 0;
 #endif
 }
 
diff --git a/Engine/Plugins/Experimental/Animation/AnimGen/Source/AnimGen/Private/RigVMFunction_AnimGen.cpp b/Engine/Plugins/Experimental/Animation/AnimGen/Source/AnimGen/Private/RigVMFunction_AnimGen.cpp
--- a/Engine/Plugins/Experimental/Animation/AnimGen/Source/AnimGen/Private/RigVMFunction_AnimGen.cpp
+++ b/Engine/Plugins/Experimental/Animation/AnimGen/Source/AnimGen/Private/RigVMFunction_AnimGen.cpp
@@ -38,11 +38,10 @@ FRigVMFunction_MakeNamedInclusiveDiscreteControl_Execute()
 
 FRigVMFunction_MakeTrajectoryFollowBehaviorControl_Execute()
 {
-	Result = UAnimGenBehaviorLibrary::MakeTrajectoryFollowBehaviorControlFromArrayView(
+	Result = UAnimGenBehaviorLibrary::MakeTrajectoryFollowBehaviorControl(
 		Object,
 		TransformTrajectory,
 		RootBoneTransform,
-		DesiredTags,
 		TrajectorySampleNum,
 		TrajectoryFutureTime,
 		TrajectoryPastTime,
diff --git a/Engine/Plugins/Experimental/Animation/AnimGen/Source/AnimGen/Public/AnimGenBehavior.h b/Engine/Plugins/Experimental/Animation/AnimGen/Source/AnimGen/Public/AnimGenBehavior.h
--- a/Engine/Plugins/Experimental/Animation/AnimGen/Source/AnimGen/Public/AnimGenBehavior.h
+++ b/Engine/Plugins/Experimental/Animation/AnimGen/Source/AnimGen/Public/AnimGenBehavior.h
@@ -21,7 +21,7 @@ class UAnimGenAutoEncoder;
  * To define a behavior it is required to inherit from this class and to provide the "SpecifyControl", and "MakeControlSets"
  * functions. The "DrawDebugControl" function may also be provided to do debug drawing in the training viewport.
  */
-UCLASS(Abstract, HideDropdown, EditInlineNew, DefaultToInstanced, BlueprintType, Blueprintable)
+UCLASS(Abstract, HideDropdown, EditInlineNew, DefaultToInstanced, CollapseCategories, BlueprintType, Blueprintable)
 class UAnimGenBehavior : public UObject
 {
 	GENERATED_BODY()
@@ -56,212 +56,6 @@ public:
 #endif
 };
 
-/** Struct to write out when resetting the preview scene of a behavior */
-USTRUCT(BlueprintType)
-struct FAnimGenBehaviorPreviewReset
-{
-	GENERATED_BODY()
-
-	/** Sequence index to reset to */
-	UPROPERTY(EditAnywhere, BlueprintReadWrite, Category = "AnimGen")
-	int32 ResetSequenceIdx = 0;
-
-	/** Time in the given sequence to reset to */
-	UPROPERTY(EditAnywhere, BlueprintReadWrite, Category = "AnimGen")
-	float ResetSequenceTime = 0.0f;
-
-	/** Root location to reset to */
-	UPROPERTY(EditAnywhere, BlueprintReadWrite, Category = "AnimGen")
-	FVector ResetLocation = FVector::ZeroVector;
-
-	/** Root rotation to reset to */
-	UPROPERTY(EditAnywhere, BlueprintReadWrite, Category = "AnimGen")
-	FRotator ResetRotation = FRotator::ZeroRotator;
-};
-
-
-/**
- * This class can be implemented to perform preview scene logic when testing an AnimGenController
- */
-UCLASS(Abstract, HideDropdown, EditInlineNew, DefaultToInstanced, BlueprintType, Blueprintable)
-class UAnimGenBehaviorPreview : public UObject
-{
-	GENERATED_BODY()
-
-#if WITH_EDITOR
-
-public:
-
-	/** Make the controls for the given state of this instance of the behavior */
-	UFUNCTION(BlueprintNativeEvent, BlueprintPure, Category = "AnimGen", meta = (ForceAsFunction))
-	UE_API FAnimGenControlObjectElement MakeControl(UPARAM(ref) FAnimGenControlObject& InControlObject, const FAnimDatabasePoseState& InPoseState) const;
-
-public:
-
-	/** Initialize the preview setup with the given Controller and AutoEncoder */
-	UFUNCTION(BlueprintNativeEvent, BlueprintCallable, Category = "AnimGen")
-	UE_API void OnPreviewInitialize(const UAnimGenController* Controller, const UAnimGenAutoEncoder* AutoEncoder);
-
-	/** Returns true if the preview scene requires reset. Provides the reset data OutResetData if required */
-	UFUNCTION(BlueprintNativeEvent, BlueprintCallable, Category = "AnimGen Preview Scene")
-	UE_API bool OnPreviewCheckReset(FAnimGenBehaviorPreviewReset& OutResetData, const FAnimDatabasePoseState& InPoseState);
-
-	/** Event called when the preview scene is reset */
-	UFUNCTION(BlueprintNativeEvent, BlueprintCallable, Category = "AnimGen Preview Scene")
-	UE_API void OnPreviewReset(const FAnimDatabasePoseState& InPoseState);
-
-	/** Event called when the preview scene is update each frame */
-	UFUNCTION(BlueprintNativeEvent, BlueprintCallable, Category = "AnimGen Preview Scene")
-	UE_API void OnPreviewUpdate(const FAnimDatabasePoseState& InPoseState, const float DeltaTime);
-
-	/** Event called each frame where debug draw logic can be added to the preview scene */
-	UFUNCTION(BlueprintNativeEvent, BlueprintCallable, Category = "AnimGen Preview Scene")
-	UE_API void OnPreviewDrawDebug(
-		const FDebugDrawer& Drawer,
-		const FAnimGenControlObject& InControlObject,
-		const FAnimGenControlObjectElement InControlObjectElement,
-		const FAnimDatabasePoseState& InPoseState,
-		const FLinearColor& IdentifierColor);
-
-	/** Event called when the left stick is moved in the preview scene */
-	UFUNCTION(BlueprintNativeEvent, BlueprintCallable, Category = "AnimGen Preview Scene")
-	UE_API void OnPreviewInputLeftStickAxis(const FVector Axis);
-
-	/** Event called when the right stick is moved in the preview scene */
-	UFUNCTION(BlueprintNativeEvent, BlueprintCallable, Category = "AnimGen Preview Scene")
-	UE_API void OnPreviewInputRightStickAxis(const FVector Axis);
-
-	/** Event called when the left trigger is moved in the preview scene */
-	UFUNCTION(BlueprintNativeEvent, BlueprintCallable, Category = "AnimGen Preview Scene")
-	UE_API void OnPreviewInputLeftTriggerAxis(const float Trigger);
-
-	/** Event called when the right trigger is moved in the preview scene */
-	UFUNCTION(BlueprintNativeEvent, BlueprintCallable, Category = "AnimGen Preview Scene")
-	UE_API void OnPreviewInputRightTriggerAxis(const float Trigger);
-
-	/** Event called when the mouse is moved in the preview scene */
-	UFUNCTION(BlueprintNativeEvent, BlueprintCallable, Category = "AnimGen Preview Scene")
-	UE_API void OnPreviewInputMouseMoved(const FVector MouseGroundLocation);
-
-	/** Event called when the mouse left button is pressed in the preview scene */
-	UFUNCTION(BlueprintNativeEvent, BlueprintCallable, Category = "AnimGen Preview Scene")
-	UE_API void OnPreviewInputLeftMousePressed(const FVector MouseGroundLocation);
-
-	/** Event called when the mouse left button is released in the preview scene */
-	UFUNCTION(BlueprintNativeEvent, BlueprintCallable, Category = "AnimGen Preview Scene")
-	UE_API void OnPreviewInputLeftMouseReleased(const FVector MouseGroundLocation);
-
-	/** Event called when the mouse right button is pressed in the preview scene */
-	UFUNCTION(BlueprintNativeEvent, BlueprintCallable, Category = "AnimGen Preview Scene")
-	UE_API void OnPreviewInputRightMousePressed(const FVector MouseGroundLocation);
-
-	/** Event called when the mouse right button is released in the preview scene */
-	UFUNCTION(BlueprintNativeEvent, BlueprintCallable, Category = "AnimGen Preview Scene")
-	UE_API void OnPreviewInputRightMouseReleased(const FVector MouseGroundLocation);
-
-	/** Event called when a key is pressed in the preview scene */
-	UFUNCTION(BlueprintNativeEvent, BlueprintCallable, Category = "AnimGen Preview Scene")
-	UE_API void OnPreviewInputKeyPressed(const FKey Key);
-
-	/** Event called when a key is released in the preview scene */
-	UFUNCTION(BlueprintNativeEvent, BlueprintCallable, Category = "AnimGen Preview Scene")
-	UE_API void OnPreviewInputKeyReleased(const FKey Key);
-
-public:
-
-	/** Gets the preview camera location and rotation */
-	UFUNCTION(BlueprintPure = true, Category = "AnimGen Preview Scene")
-	UE_API void GetPreviewCameraLocationAndRotation(FVector& OutCameraLocation, FRotator& OutCameraRotation) const;
-
-	/** Gets the preview camera yaw rotation */
-	UFUNCTION(BlueprintPure = true, Category = "AnimGen Preview Scene")
-	UE_API FRotator GetPreviewCameraYawRotation() const;
-
-	/** Updates the preview camera using the provided targets */
-	UFUNCTION(BlueprintPure = false, Category = "AnimGen Preview Scene")
-	UE_API void UpdatePreviewCamera(
-		const float YawVelocity,
-		const float PitchVelocity,
-		const float DistanceVelocity,
-		const float DeltaTime,
-		const bool bInvertPitch = false);
-
-	/** Updates the preview camera using the shoulder buttons */
-	UFUNCTION(BlueprintPure = false, Category = "AnimGen Preview Scene")
-	UE_API void UpdatePreviewCameraFromAxisAndShoulders(
-		const FVector Axis,
-		const bool bLeftShoulder,
-		const bool bRightShoulder,
-		const float DeltaTime,
-		const float RotationSpeed = 180.0f,
-		const float DistanceSpeed = 200.0f,
-		const bool bInvertPitch = false,
-		const float Deadzone = 0.2f,
-		const float Sensitivity = 2.0f);
-
-	/** Gets the preview scene left stick axis */
-	UFUNCTION(BlueprintPure = true, Category = "AnimGen Preview Scene")
-	UE_API FVector GetPreviewLeftStickAxis() const;
-
-	/** Gets the preview scene right stick axis */
-	UFUNCTION(BlueprintPure = true, Category = "AnimGen Preview Scene")
-	UE_API FVector GetPreviewRightStickAxis() const;
-
-	/** Gets the preview scene camera-relative left stick axis */
-	UFUNCTION(BlueprintPure = true, Category = "AnimGen Preview Scene")
-	UE_API FVector GetPreviewCameraRelativeLeftStickAxis() const;
-
-	/** Gets the preview scene camera-relative right stick axis */
-	UFUNCTION(BlueprintPure = true, Category = "AnimGen Preview Scene")
-	UE_API FVector GetPreviewCameraRelativeRightStickAxis() const;
-
-	/** Gets the preview scene left trigger axis */
-	UFUNCTION(BlueprintPure = true, Category = "AnimGen Preview Scene")
-	UE_API float GetPreviewLeftTriggerAxis() const;
-
-	/** Gets the preview scene right trigger axis */
-	UFUNCTION(BlueprintPure = true, Category = "AnimGen Preview Scene")
-	UE_API float GetPreviewRightTriggerAxis() const;
-
-	/** Gets the preview scene left trigger axis with the given sensitivity */
-	UFUNCTION(BlueprintPure = true, Category = "AnimGen Preview Scene")
-	UE_API float GetPreviewLeftTriggerAxisWithSensitivity(const float Sensitivity = 2.0f) const;
-
-	/** Gets the preview scene right trigger axis with the given sensitivity */
-	UFUNCTION(BlueprintPure = true, Category = "AnimGen Preview Scene")
-	UE_API float GetPreviewRightTriggerAxisWithSensitivity(const float Sensitivity = 2.0f) const;
-
-	/** Gets the preview scene mouse location on the ground plane */
-	UFUNCTION(BlueprintPure = true, Category = "AnimGen Preview Scene")
-	UE_API FVector GetPreviewMouseGroundLocation() const;
-
-#endif
-
-public:
-
-#if WITH_EDITORONLY_DATA
-
-	float PreviewCameraHeight = 110.0f;
-	float PreviewCameraYaw = -90.0f;
-	float PreviewCameraPitch = 30.0f;
-	float PreviewCameraDistance = 300.0f;
-	FVector PreviewCameraLocation = FVector::ZeroVector;
-	FRotator PreviewCameraRotation = FRotator::ZeroRotator;
-
-	FVector PreviewLeftStick = FVector::ZeroVector;
-	FVector PreviewRightStick = FVector::ZeroVector;
-	float PreviewLeftTrigger = 0.0f;
-	float PreviewRightTrigger = 0.0f;
-	FVector PreviewPrevLeftStick = FVector::ZeroVector;
-	FVector PreviewPrevRightStick = FVector::ZeroVector;
-	float PreviewPrevLeftTrigger = 0.0f;
-	float PreviewPrevRightTrigger = 0.0f;
-	FVector PreviewMouseGroundLocation = FVector::ZeroVector;
-	FVector PreviewPrevMouseGroundLocation = FVector::ZeroVector;
-
-#endif
-};
-
 /** Behavior for uncontrolled generation. Produces a random-walk of the dataset. */
 UCLASS(BlueprintType, Blueprintable, meta = (DisplayName = "Uncontrolled"))
 class UAnimGenBehavior_Uncontrolled : public UAnimGenBehavior
@@ -279,44 +73,11 @@ public:
 #endif
 };
 
-/** Behavior for idling - user can provide some tags to control the style of the idle. */
+/** Behavior for idling. */
 UCLASS(BlueprintType, Blueprintable, meta = (DisplayName = "Idle"))
-class UAnimGenBehavior_Idle : public UAnimGenBehavior
+class UAnimGenBehavior_Idle : public UAnimGenBehavior_Uncontrolled
 {
 	GENERATED_BODY()
-
-#if WITH_EDITORONLY_DATA
-
-public:
-
-	/** List of tags with associated frame ranges */
-	UPROPERTY(EditAnywhere, Category = "Tags")
-	TArray<FAnimDatabaseFrameRangesEntry> TagRanges;
-
-#endif
-
-#if WITH_EDITOR
-
-public:
-
-	UE_API virtual FAnimGenControlSchemaElement SpecifyControl_Implementation(FAnimGenControlSchema& InControlSchema) const override;
-
-	UE_API virtual void MakeControlSets_Implementation(TArray<FAnimGenControlSet>& OutControlSets, FAnimGenControlObject& InControlObject, const UAnimDatabase* InDatabase, const FAnimDatabaseFrameRanges& InFrameRanges) const override;
-
-	UE_API virtual void DrawDebugControl_Implementation(
-		const FDebugDrawer& Drawer,
-		const FAnimGenControlObject& InControlObject,
-		const FAnimGenControlObjectElement InControlObjectElement,
-		const FAnimDatabasePoseState& InPoseState,
-		const UAnimDatabase* InDatabase,
-		const UAnimGenAutoEncoder* InAutoEncoder,
-		const int32 SequenceIdx,
-		const float SequenceTime,
-		const int32 RangeStart,
-		const int32 RangeLength,
-		const FLinearColor& IdentifierColor) const override;
-
-#endif
 };
 
 /** Behavior for following a trajectory which can either be generated from Mover/CharacterMovementComponent or a path to follow */
@@ -330,25 +91,21 @@ class UAnimGenBehavior_TrajectoryFollow : public UAnimGenBehavior
 public:
 
 	/** Future lookahead time for the trajectory */
-	UPROPERTY(EditAnywhere, Category = "Trajectory")
+	UPROPERTY(EditAnywhere, Category = "Settings")
 	float TrajectoryFutureTime = 1.0f;
 
 	/** Past time for the trajectory */
-	UPROPERTY(EditAnywhere, Category = "Trajectory")
+	UPROPERTY(EditAnywhere, Category = "Settings")
 	float TrajectoryPastTime = 0.0f;
 
 	/** Number of samples to use for the trajectory */
-	UPROPERTY(EditAnywhere, Category = "Trajectory")
+	UPROPERTY(EditAnywhere, Category = "Settings")
 	int32 TrajectorySampleNum = 4;
 
 	/** Local root forward vector for the trajectory */
-	UPROPERTY(EditAnywhere, Category = "Trajectory")
+	UPROPERTY(EditAnywhere, Category = "Settings")
 	FVector TrajectoryForwardVector = FVector::RightVector;
 
-	/** List of tags with associated frame ranges */
-	UPROPERTY(EditAnywhere, Category = "Tags")
-	TArray<FAnimDatabaseFrameRangesEntry> TagRanges;
-
 #endif
 
 #if WITH_EDITOR
@@ -381,16 +138,6 @@ class UAnimGenBehavior_VelocityFollow : public UAnimGenBehavior
 {
 	GENERATED_BODY()
 
-#if WITH_EDITORONLY_DATA
-
-public:
-
-	/** List of tags with associated frame ranges */
-	UPROPERTY(EditAnywhere, Category = "Tags")
-	TArray<FAnimDatabaseFrameRangesEntry> TagRanges;
-
-#endif
-
 public:
 
 #if WITH_EDITOR
@@ -421,16 +168,6 @@ class UAnimGenBehavior_MoveToTarget : public UAnimGenBehavior
 {
 	GENERATED_BODY()
 
-#if WITH_EDITORONLY_DATA
-
-public:
-
-	/** List of tags with associated frame ranges */
-	UPROPERTY(EditAnywhere, Category = "Tags")
-	TArray<FAnimDatabaseFrameRangesEntry> TagRanges;
-
-#endif
-
 #if WITH_EDITOR
 
 public:
@@ -459,7 +196,7 @@ public:
  * Behavior for following a trajectory while also interacting with an object or prop. This is useful for networked prop interactions where the 
  * trajectory may be coming from a montage .
  */
-UCLASS(BlueprintType, Blueprintable, meta = (DisplayName = "Trajectory Interaction"))
+UCLASS(BlueprintType, Blueprintable, DontCollapseCategories, meta = (DisplayName = "Trajectory Interaction"))
 class UAnimGenBehavior_TrajectoryInteraction : public UAnimGenBehavior
 {
 	GENERATED_BODY()
@@ -484,10 +221,6 @@ public:
 	UPROPERTY(EditAnywhere, Category = "Trajectory")
 	FVector TrajectoryForwardVector = FVector::RightVector;
 
-	/** List of tags with associated frame ranges */
-	UPROPERTY(EditAnywhere, Category = "Tags")
-	TArray<FAnimDatabaseFrameRangesEntry> TagRanges;
-
 	/** Name of bone used for interaction */
 	UPROPERTY(EditAnywhere, Category = "Interaction")
 	FName InteractionBoneName = TEXT("attach");
@@ -534,6 +267,54 @@ public:
 #endif
 };
 
+/** Behavior that allows a user to supply tags to control the style of a sub-behavior. */
+UCLASS(BlueprintType, Blueprintable, meta = (DisplayName = "Tagged"))
+class UAnimGenBehavior_Tagged : public UAnimGenBehavior
+{
+	GENERATED_BODY()
+
+#if WITH_EDITORONLY_DATA
+
+public:
+
+	/** Sub-behavior */
+	UPROPERTY(EditAnywhere, Instanced, Category = "Settings")
+	TObjectPtr<UAnimGenBehavior> Behavior;
+
+	/** List of tags with associated frame ranges */
+	UPROPERTY(EditAnywhere, Category = "Settings")
+	TArray<FAnimDatabaseFrameRangesEntry> TagRanges;
+
+	/** Offset to debug draw the current tags */
+	UPROPERTY(EditAnywhere, Category = "Settings")
+	FVector TagDebugDrawOffset = FVector(0.0, 0.0, 220.0f);
+
+#endif
+
+#if WITH_EDITOR
+
+public:
+
+	UE_API virtual FAnimGenControlSchemaElement SpecifyControl_Implementation(FAnimGenControlSchema& InControlSchema) const override;
+
+	UE_API virtual void MakeControlSets_Implementation(TArray<FAnimGenControlSet>& OutControlSets, FAnimGenControlObject& InControlObject, const UAnimDatabase* InDatabase, const FAnimDatabaseFrameRanges& InFrameRanges) const override;
+
+	UE_API virtual void DrawDebugControl_Implementation(
+		const FDebugDrawer& Drawer,
+		const FAnimGenControlObject& InControlObject,
+		const FAnimGenControlObjectElement InControlObjectElement,
+		const FAnimDatabasePoseState& InPoseState,
+		const UAnimDatabase* InDatabase,
+		const UAnimGenAutoEncoder* InAutoEncoder,
+		const int32 SequenceIdx,
+		const float SequenceTime,
+		const int32 RangeStart,
+		const int32 RangeLength,
+		const FLinearColor& IdentifierColor) const override;
+
+#endif
+};
+
 /** Entry specifying a behavior */
 USTRUCT(BlueprintType)
 struct FAnimGenBehaviorEntry
@@ -566,13 +347,9 @@ class UAnimGenBehavior_Multi : public UAnimGenBehavior
 public:
 
 	/** Different potential sub-behaviors */
-	UPROPERTY(EditAnywhere, Category = "Behaviors")
+	UPROPERTY(EditAnywhere, Category = "Settings")
 	TArray<FAnimGenBehaviorEntry> Behaviors;
 
-	/** List of tags with associated frame ranges */
-	UPROPERTY(EditAnywhere, Category = "Tags")
-	TArray<FAnimDatabaseFrameRangesEntry> TagRanges;
-
 #endif
 
 #if WITH_EDITOR
@@ -599,34 +376,6 @@ public:
 #endif
 };
 
-
-/** Simple Trajectory type that stores all variables required for simple spring-based prediction. */
-USTRUCT(BlueprintType)
-struct FAnimGenTrajectory
-{
-	GENERATED_BODY()
-
-public:
-
-	UPROPERTY(VisibleAnywhere, Category = "Trajectory")
-	TArray<FVector> Locations;
-
-	UPROPERTY(VisibleAnywhere, Category = "Trajectory")
-	TArray<FVector> Velocities;
-
-	UPROPERTY(VisibleAnywhere, Category = "Trajectory")
-	TArray<FVector> Accelerations;
-
-	UPROPERTY(VisibleAnywhere, Category = "Trajectory")
-	TArray<FRotator> Rotations;
-
-	UPROPERTY(VisibleAnywhere, Category = "Trajectory")
-	TArray<FVector> AngularVelocities;
-
-	UPROPERTY(VisibleAnywhere, Category = "Trajectory")
-	TArray<FVector> Directions;
-};
-
 /** Blueprint library of helper functions for making and testing behaviors including functions for preparing, processing, and making controls */
 UCLASS(BlueprintType, meta = (BlueprintThreadSafe))
 class UAnimGenBehaviorLibrary : public UBlueprintFunctionLibrary
@@ -640,57 +389,35 @@ public:
 	static UE_API FAnimGenControlObjectElement MakeUncontrolledBehaviorControl(UPARAM(ref) FAnimGenControlObject& InControlObject);
 
 	/** Make a control object element for the Idle Behavior */
-	UFUNCTION(BlueprintPure, Category = "AnimGen", meta = (AutoCreateRefTerm = "DesiredTags"))
-	static UE_API FAnimGenControlObjectElement MakeIdleBehaviorControl(UPARAM(ref) FAnimGenControlObject& InControlObject, const TArray<FName>& DesiredTags);
-	static UE_API FAnimGenControlObjectElement MakeIdleBehaviorControlFromArrayView(FAnimGenControlObject& InControlObject, const TArrayView<const FName> DesiredTags = {});
+	UFUNCTION(BlueprintPure, Category = "AnimGen")
+	static UE_API FAnimGenControlObjectElement MakeIdleBehaviorControl(UPARAM(ref) FAnimGenControlObject& InControlObject);
 
 	/** Make a control object element for the Trajectory Following Behavior */
-	UFUNCTION(BlueprintPure, Category = "AnimGen", meta = (AutoCreateRefTerm = "TransformTrajectory, RootBoneTransform, DesiredTags"))
+	UFUNCTION(BlueprintPure, Category = "AnimGen", meta = (AutoCreateRefTerm = "TransformTrajectory, RootBoneTransform"))
 	static UE_API FAnimGenControlObjectElement MakeTrajectoryFollowBehaviorControl(
 		UPARAM(ref) FAnimGenControlObject& InControlObject,
 		const FTransformTrajectory& TransformTrajectory,
 		const FTransform& RootBoneTransform,
-		const TArray<FName>& DesiredTags,
-		const int32 InTrajectorySampleNum = 4,
-		const float InTrajectoryFutureTime = 1.0f,
-		const float InTrajectoryPastTime = 0.0f,
-		const FVector InTrajectoryForwardVector = FVector::RightVector);
-
-	static UE_API FAnimGenControlObjectElement MakeTrajectoryFollowBehaviorControlFromArrayView(
-		FAnimGenControlObject& InControlObject,
-		const FTransformTrajectory& TransformTrajectory,
-		const FTransform& RootBoneTransform,
-		const TArrayView<const FName> DesiredTags = {},
 		const int32 InTrajectorySampleNum = 4,
 		const float InTrajectoryFutureTime = 1.0f,
 		const float InTrajectoryPastTime = 0.0f,
 		const FVector InTrajectoryForwardVector = FVector::RightVector);
 
 	/** Make a control object element for the Move To Target Behavior */
-	UFUNCTION(BlueprintPure, Category = "AnimGen", meta = (AutoCreateRefTerm = "RootBoneTransform, DesiredTags"))
+	UFUNCTION(BlueprintPure, Category = "AnimGen", meta = (AutoCreateRefTerm = "RootBoneTransform"))
 	static UE_API FAnimGenControlObjectElement MakeMoveToTargetBehaviorControl(
 		UPARAM(ref) FAnimGenControlObject& InControlObject,
 		const FVector TargetLocation,
 		const FVector TargetDirection,
 		const bool bUseTargetDirection,
-		const FTransform& RootBoneTransform,
-		const TArray<FName>& DesiredTags);
-
-	static UE_API FAnimGenControlObjectElement MakeMoveToTargetBehaviorControlFromArrayView(
-		FAnimGenControlObject& InControlObject,
-		const FVector TargetLocation,
-		const FVector TargetDirection,
-		const bool bUseTargetDirection,
-		const FTransform& RootBoneTransform,
-		const TArrayView<const FName> DesiredTags = {});
+		const FTransform& RootBoneTransform);
 
 	/** Make a control object element for the Trajectory Interaction Behavior */
-	UFUNCTION(BlueprintPure, Category = "AnimGen", meta = (AutoCreateRefTerm = "TransformTrajectory, RootBoneTransform, DesiredTags"))
+	UFUNCTION(BlueprintPure, Category = "AnimGen", meta = (AutoCreateRefTerm = "TransformTrajectory, RootBoneTransform"))
 	static UE_API FAnimGenControlObjectElement MakeTrajectoryInteractionBehaviorControl(
 		UPARAM(ref) FAnimGenControlObject& InControlObject,
 		const FTransformTrajectory& TransformTrajectory,
 		const FTransform& RootBoneTransform,
-		const TArray<FName>& DesiredTags,
 		const FVector InteractionLocation = FVector::ZeroVector,
 		const FVector InteractionDirection = FVector::ForwardVector,
 		const bool bInteractionActive = false,
@@ -699,103 +426,24 @@ public:
 		const float InTrajectoryPastTime = 0.0f,
 		const FVector InTrajectoryForwardVector = FVector::RightVector);
 
-	static UE_API FAnimGenControlObjectElement MakeTrajectoryInteractionBehaviorControlFromArrayView(
-		FAnimGenControlObject& InControlObject,
-		const FTransformTrajectory& TransformTrajectory,
-		const FTransform& RootBoneTransform,
-		const TArrayView<const FName> DesiredTags = {},
-		const FVector InteractionLocation = FVector::ZeroVector,
-		const FVector InteractionDirection = FVector::ForwardVector,
-		const bool bInteractionActive = false,
-		const int32 InTrajectorySampleNum = 4,
-		const float InTrajectoryFutureTime = 1.0f,
-		const float InTrajectoryPastTime = 0.0f,
-		const FVector InTrajectoryForwardVector = FVector::RightVector);
-
-	/** Make a control object element for the Multi Behavior */
+	/** Make a control object element for the Tagged Behavior */
 	UFUNCTION(BlueprintPure, Category = "AnimGen", meta = (AutoCreateRefTerm = "DesiredTags"))
-	static UE_API FAnimGenControlObjectElement MakeMultiBehaviorControl(
+	static UE_API FAnimGenControlObjectElement MakeTaggedBehaviorControl(
 		UPARAM(ref) FAnimGenControlObject& InControlObject,
-		const FName Behavior,
 		const FAnimGenControlObjectElement& BehaviorControlElement,
 		const TArray<FName>& DesiredTags);
 
-	static UE_API FAnimGenControlObjectElement MakeMultiBehaviorControlFromArrayView(
+	static UE_API FAnimGenControlObjectElement MakeTaggedBehaviorControlFromArrayView(
 		FAnimGenControlObject& InControlObject,
-		const FName Behavior,
 		const FAnimGenControlObjectElement& BehaviorControlElement,
 		const TArrayView<const FName> DesiredTags);
 
-public:
-
-	/** Make a random bool value */
-	UFUNCTION(BlueprintPure, Category = "AnimGen")
-	static UE_API bool MakeRandomBool(UPARAM(ref) int32& State);
-
-	/** Make a random bool value with the given true probability */
-	UFUNCTION(BlueprintPure, Category = "AnimGen")
-	static UE_API bool MakeRandomBoolWithProbability(UPARAM(ref) int32& State, const float TrueProbability = 0.5f);
-
-	/** Switch on a randomly generated bool value */
-	UFUNCTION(BlueprintCallable, Category = "AnimGen", meta = (ExpandBoolAsExecs = "ReturnValue"))
-	static UE_API bool SwitchOnRandomBool(UPARAM(ref) int32& State);
-
-	/** Switch on a randomly generated bool value with the given true probability */
-	UFUNCTION(BlueprintCallable, Category = "AnimGen", meta = (ExpandBoolAsExecs = "ReturnValue"))
-	static UE_API bool SwitchOnRandomBoolWithProbability(UPARAM(ref) int32& State, const float TrueProbability = 0.5f);
-
-	/** Make a random float in the range [0, 1] */
-	UFUNCTION(BlueprintPure, Category = "AnimGen")
-	static UE_API float MakeRandomFloat(UPARAM(ref) int32& State);
-
-	/** Make a random float in the provided range */
-	UFUNCTION(BlueprintPure, Category = "AnimGen")
-	static UE_API float MakeRandomFloatInRange(UPARAM(ref) int32& State, const float Min = 0.0f, const float Max = 1.0f);
-
-	/** Randomly choose from a given set of floats */
-	UFUNCTION(BlueprintPure, Category = "AnimGen")
-	static UE_API float MakeRandomFloatChoice(UPARAM(ref) int32& State, const TArray<float>& Values);
-	static UE_API float MakeRandomFloatChoiceFromArrayView(int32& State, const TArrayView<const float> Values);
-
-	/** Sample an integer in the given range (inclusive) */
-	UFUNCTION(BlueprintPure, Category = "AnimGen")
-	static UE_API int32 MakeRandomIntegerInRangeInclusive(UPARAM(ref) int32& State, const int32 Min, const int32 Max);
-
-	/** Sample a random enum value from the given enum class */
-	UFUNCTION(BlueprintPure, Category = "AnimGen")
-	static UE_API uint8 MakeRandomEnumValue(UPARAM(ref) int32& State, const UEnum* Enum);
-
-	/** Sample a random 3d vector from a Gaussian Distribution */
-	UFUNCTION(BlueprintPure, Category = "AnimGen")
-	static UE_API FVector MakeRandomVectorFromGaussianDistribution(
-		UPARAM(ref) int32& State,
-		const FVector Mean = FVector::ZeroVector,
-		const FVector Std = FVector(1, 1, 1));
-
-	/** Sample a random 3d direction */
-	UFUNCTION(BlueprintPure, Category = "AnimGen")
-	static UE_API FVector MakeRandomDirection3D(
-		UPARAM(ref) int32& State,
-		const FVector Mean = FVector::ZeroVector,
-		const FVector Std = FVector(1, 1, 1),
-		const FVector Default = FVector::ForwardVector);
-
-	/** Sample a random 2d direction (the Z coordinate will be zero) */
-	UFUNCTION(BlueprintPure, Category = "AnimGen")
-	static UE_API FVector MakeRandomDirection2D(UPARAM(ref) int32& State);
-
-	/** Sample a random 2d location from a Gaussian Distribution (the Z coordinate will be zero) */
+	/** Make a control object element for the Multi Behavior */
 	UFUNCTION(BlueprintPure, Category = "AnimGen")
-	static UE_API FVector MakeRandomLocationFromGaussian2D(UPARAM(ref) int32& State, const FVector Mean = FVector::ZeroVector, const float Std = 100.0f);
-
-	/** Sample a random sequence index and sequence time from a set of frame range */
-	UFUNCTION(BlueprintPure = false, Category = "AnimGen")
-	static UE_API void MakeRandomSequenceIndexAndTimeFromFrameRanges(
-		int32& OutSequenceIdx,
-		float& OutSequenceTime,
-		UPARAM(ref) int32& State,
-		const FAnimDatabaseFrameRanges& FrameRanges,
-		const float FrameDeltaTime);
+	static UE_API FAnimGenControlObjectElement MakeMultiBehaviorControl(
+		UPARAM(ref) FAnimGenControlObject& InControlObject,
+		const FName Behavior,
+		const FAnimGenControlObjectElement& BehaviorControlElement);
 
 public:
 
@@ -908,98 +556,6 @@ public:
 		const float LocationThreshold = 25.0f,
 		const float DirectionThresoldAngle = 10.0f);
 
-public:
-
-	/** Get the velocity along the path at the given time */
-	UFUNCTION(BlueprintPure = true, Category = "AnimGen")
-	static UE_API FVector PathVelocityAtTime(const TArray<FVector>& Points, const float Time);
-	static UE_API FVector PathVelocityAtTimeFromArrayView(const TArrayView<const FVector> Points, const float Time);
-
-	/** Get the location along the path at the given time */
-	UFUNCTION(BlueprintPure = true, Category = "AnimGen")
-	static UE_API FVector PathLocationAtTime(const TArray<FVector>& Points, const float Time);
-	static UE_API FVector PathLocationAtTimeFromArrayView(const TArrayView<const FVector> Points, const float Time);
-
-	/** Get the direction along the path at the given time */
-	UFUNCTION(BlueprintPure = true, Category = "AnimGen")
-	static UE_API FVector PathDirectionAtTime(const TArray<FVector>& Points, const float Time);
-	static UE_API FVector PathDirectionAtTimeFromArrayView(const TArrayView<const FVector> Points, const float Time);
-
-	/** Get the rotation along the path at the given time */
-	UFUNCTION(BlueprintPure = true, Category = "AnimGen")
-	static UE_API FRotator PathRotationAtTime(const TArray<FVector>& Points, const float Time);
-	static UE_API FRotator PathRotationAtTimeFromArrayView(const TArrayView<const FVector> Points, const float Time);
-
-	/** Adds a point to a path if it is further away than the given minimum distance from the last point added */
-	UFUNCTION(BlueprintPure = false, Category = "AnimGen", meta = (ReturnDisplayName = "Added"))
-	static UE_API bool PathAddPoint(UPARAM(ref) TArray<FVector>& Points, const FVector NewPoint, const float MinDistance = 50.0f);
-
-	/** Get the total distance of a path */
-	UFUNCTION(BlueprintPure = true, Category = "AnimGen")
-	static UE_API float PathTotalDistance(const TArray<FVector>& Points);
-	static UE_API float PathTotalDistanceFromArrayView(const TArrayView<const FVector> Points);
-
-public:
-
-	/** Make a transform trajectory from an array of locations, rotations, and times */
-	UFUNCTION(BlueprintPure = true, Category = "AnimGen")
-	static UE_API void MakeTransformTrajectory(FTransformTrajectory& OutTransformTrajectory, const TArray<FVector>& Locations, const TArray<FRotator>& Rotations, const TArray<float>& Times);
-	static UE_API void MakeTransformTrajectoryFromArrayViews(FTransformTrajectory& OutTransformTrajectory, const TArrayView<const FVector> Locations, const TArrayView<const FRotator> Rotations, const TArrayView<const float> Times);
-
-	/** Make a transform trajectory from an AnimGenTrajectory */
-	UFUNCTION(BlueprintPure = true, Category = "AnimGen")
-	static UE_API void MakeTransformTrajectoryFromTrajectory(FTransformTrajectory& OutTransformTrajectory, const FAnimGenTrajectory& Trajectory, const float FutureTime = 1.0f, const float PastTime = 0.0f);
-
-	/** Make a transform trajectory from a path */
-	UFUNCTION(BlueprintPure = true, Category = "AnimGen")
-	static UE_API void MakeTransformTrajectoryFromPath(FTransformTrajectory& OutTransformTrajectory, const TArray<FVector>& Path, const float StartDistance = 0.0f, const float PathPointDistance = 50.0f, const float PathSpeed = 150.0f, const float FutureTime = 1.0f, const float PastTime = 0.0f, const float SampleInterval = 0.01666666666f);
-	static UE_API void MakeTransformTrajectoryFromPathArrayView(FTransformTrajectory& OutTransformTrajectory, const TArrayView<const FVector> Path, const float StartDistance = 0.0f, const float PathPointDistance = 50.0f, const float PathSpeed = 150.0f, const float FutureTime = 1.0f, const float PastTime = 0.0f, const float SampleInterval = 0.01666666666f);
-
-	/** Project a transform trajectory onto the ground */
-	UFUNCTION(BlueprintPure = false, Category = "AnimGen")
-	static UE_API void ProjectTransformTrajectoryOntoFloor(UPARAM(ref) FTransformTrajectory& InOutTransformTrajectory, const float FloorHeight = 0.0f);
-
-	/** Get the number of samples in an AnimGenTrajectory */
-	UFUNCTION(BlueprintPure = true, Category = "AnimGen")
-	static UE_API int32 GetTrajectorySampleNum(const FAnimGenTrajectory& Trajectory);
-
-	/** Get the locations in an AnimGenTrajectory */
-	UFUNCTION(BlueprintPure = true, Category = "AnimGen")
-	static UE_API const TArray<FVector>& GetTrajectoryLocations(const FAnimGenTrajectory& Trajectory);
-
-	/** Get the directions in an AnimGenTrajectory */
-	UFUNCTION(BlueprintPure = true, Category = "AnimGen")
-	static UE_API const TArray<FVector>& GetTrajectoryDirections(const FAnimGenTrajectory& Trajectory);
-
-	/** Reset an AnimGenTrajectory to a given state */
-	UFUNCTION(BlueprintPure = false, Category = "AnimGen")
-	static UE_API void ResetTrajectory(
-		UPARAM(ref) FAnimGenTrajectory& Trajectory,
-		const int32 SampleNum = 4,
-		const FVector InitialLocation = FVector::ZeroVector,
-		const FRotator InitialRotation = FRotator::ZeroRotator,
-		const FVector ForwardVector = FVector(0, 1, 0));
-
-	/** Update and predict the future for an AnimGenTrajectory */
-	UFUNCTION(BlueprintPure = false, Category = "AnimGen", meta = (AdvancedDisplay = 8))
-	static UE_API void UpdateAndPredictFutureTrajectory(
-		UPARAM(ref) FAnimGenTrajectory& Trajectory,
-		const FVector DesiredVelocity,
-		const FRotator DesiredRotation,
-		const float DeltaTime,
-		const float FutureTime = 1.0f,
-		const float MaximumDesiredVelocity = 500.0f,
-		const float MaximumAcceleration = 1000.0f,
-		const float MaximumAngularVelocity = 360.0f,
-		const FVector ForwardVector = FVector(0, 1, 0));
-
-	/** Teleport an AnimGenTrajectory. Updates velocities and future estimations. */
-	UFUNCTION(BlueprintPure = false, Category = "AnimGen")
-	static UE_API void TeleportTrajectory(
-		UPARAM(ref) FAnimGenTrajectory& Trajectory,
-		const FVector NewLocation,
-		const FRotator NewRotation);
-
 public:
 
 	/** Applies a simple avoidance heuristic to a movement velocity */
@@ -1046,331 +602,4 @@ public:
 		const float DotAngleThreshold = 0.75f);
 };
 
-/** Simple Behavior preview for the uncontrolled behavior */
-UCLASS(BlueprintType, Blueprintable, meta = (DisplayName = "Uncontrolled"))
-class UAnimGenBehaviorPreview_Uncontrolled : public UAnimGenBehaviorPreview
-{
-	GENERATED_BODY()
-
-#if WITH_EDITOR
-
-public:
-
-	virtual FAnimGenControlObjectElement MakeControl_Implementation(FAnimGenControlObject& InControlObject, const FAnimDatabasePoseState& InPoseState) const override;
-
-#endif
-};
-
-/** Simple Behavior preview for the idle behavior */
-UCLASS(BlueprintType, Blueprintable, meta = (DisplayName = "Idle"))
-class UAnimGenBehaviorPreview_Idle : public UAnimGenBehaviorPreview
-{
-	GENERATED_BODY()
-
-#if WITH_EDITORONLY_DATA
-
-public:
-
-	/** List of styles to cycle through */
-	UPROPERTY(EditAnywhere, Category = "Styles")
-	TArray<FName> StyleTags;
-
-protected:
-
-	/** Current Desired Style */
-	FName DesiredStyle = NAME_None;
-
-	/** Current Style Index */
-	int32 StyleIndex = 0;
-
-#endif
-
-#if WITH_EDITOR
-
-public:
-
-	virtual FAnimGenControlObjectElement MakeControl_Implementation(FAnimGenControlObject& InControlObject, const FAnimDatabasePoseState& InPoseState) const override;
-
-	virtual void OnPreviewInputKeyPressed_Implementation(const FKey Key) override;
-
-	virtual void OnPreviewDrawDebug_Implementation(
-		const FDebugDrawer& Drawer,
-		const FAnimGenControlObject& InControlObject,
-		const FAnimGenControlObjectElement InControlObjectElement,
-		const FAnimDatabasePoseState& InPoseState,
-		const FLinearColor& IdentifierColor) override;
-#endif
-};
-
-/** Behavior preview for Trajectory Following that allows gamepad control or drawing of trajectory on the ground */
-UCLASS(BlueprintType, Blueprintable, meta = (DisplayName = "Trajectory Follow"))
-class UAnimGenBehaviorPreview_TrajectoryFollow : public UAnimGenBehaviorPreview
-{
-	GENERATED_BODY()
-
-#if WITH_EDITORONLY_DATA
-
-public:
-
-	/** Minimum speed when walking */
-	UPROPERTY(EditAnywhere, Category = "Locomotion", meta = (ForceUnits = "cm/s"))
-	float MinimumWalkSpeed = 75.0f;
-
-	/** Maximum Speed used for walking */
-	UPROPERTY(EditAnywhere, Category = "Locomotion", meta = (ForceUnits = "cm/s"))
-	float MaximumWalkSpeed = 150.0f;
-
-	/** Minimum Speed used for running */
-	UPROPERTY(EditAnywhere, Category = "Locomotion", meta = (ForceUnits = "cm/s"))
-	float MinimumRunSpeed = 200.0f;
-
-	/** Maximum Speed used for running */
-	UPROPERTY(EditAnywhere, Category = "Locomotion", meta = (ForceUnits = "cm/s"))
-	float MaximumRunSpeed = 500.0f;
-
-	/** Smoothing to apply to speed changes */
-	UPROPERTY(EditAnywhere, Category = "Locomotion", meta = (ForceUnits = "s"))
-	float SpeedSmoothingTime = 0.1f;
-
-	/** Maximum acceleration for movement model */
-	UPROPERTY(EditAnywhere, Category = "Locomotion", meta = (ForceUnits = "cm/s^2"))
-	float MaximumAcceleration = 1000.0f;
-
-	/** Maximum angular velocity used for movement model */
-	UPROPERTY(EditAnywhere, Category = "Locomotion", meta = (ForceUnits = "deg/s"))
-	float MaximumAngularVelocity = 360.0f;
-
-	/** Tag used for standing */
-	UPROPERTY(EditAnywhere, Category = "Locomotion")
-	FName StandTag = NAME_None;
-
-	/** Tag used for walking */
-	UPROPERTY(EditAnywhere, Category = "Locomotion")
-	FName WalkTag = NAME_None;
-
-	/** Tag used for running */
-	UPROPERTY(EditAnywhere, Category = "Locomotion")
-	FName RunTag = NAME_None;
-
-public:
-
-	/** Future lookahead time for the trajectory */
-	UPROPERTY(EditAnywhere, Category = "Trajectory", meta = (ForceUnits = "s"))
-	float TrajectoryFutureTime = 1.0f;
-
-	/** Past time for the trajectory */
-	UPROPERTY(EditAnywhere, Category = "Trajectory", meta = (ForceUnits = "s"))
-	float TrajectoryPastTime = 0.0f;
-
-	/** Number of samples to use for the trajectory */
-	UPROPERTY(EditAnywhere, Category = "Trajectory")
-	int32 TrajectorySampleNum = 4;
-
-	/** Local root forward vector for the trajectory */
-	UPROPERTY(EditAnywhere, Category = "Trajectory")
-	FVector TrajectoryForwardVector = FVector::RightVector;
-
-public:
-
-	/** List of styles to cycle through */
-	UPROPERTY(EditAnywhere, Category = "Styles")
-	TArray<FName> StyleTags;
-
-public:
-
-	/** Maximum speed at which the camera can rotate */
-	UPROPERTY(EditAnywhere, Category = "Camera", meta = (ForceUnits = "deg/s"))
-	float CameraRotationSpeed = 180.0f;
-
-	/** Maximum speed at which the camera can move closer and further from the character */
-	UPROPERTY(EditAnywhere, Category = "Camera", meta = (ForceUnits = "cm/s"))
-	float CameraZoomSpeed = 100.0f;
-
-	/** Smoothing to apply to camera input controls */
-	UPROPERTY(EditAnywhere, Category = "Camera", meta = (ForceUnits = "s"))
-	float CameraInputSmoothingTime = 0.1f;
-
-	/** If to invert the pitch control for the camera */
-	UPROPERTY(EditAnywhere, Category = "Camera")
-	bool bCameraInvertPitch = true;
-
-public:
-
-	/** Speed to use for path following */
-	UPROPERTY(EditAnywhere, Category = "Path", meta = (ForceUnits = "cm/s"))
-	float PathFollowSpeed = 150.0f;
-
-	/** Minimum distance between points in the path */
-	UPROPERTY(EditAnywhere, Category = "Path", meta = (ForceUnits = "cm"))
-	float PathPointDistance = 50.0f;
-
-	/** If the path facing direction should face the camera */
-	UPROPERTY(EditAnywhere, Category = "Path")
-	bool bPathFacingCamera = false;
-
-protected:
-
-	/** If to follow a path rather than the trajectory */
-	bool bFollowPath = false;
-
-	/** If a path is being drawn */
-	bool bDrawPath = false;
-
-	/** Path to follow */
-	TArray<FVector> FollowPath;
-
-	/** Path being drawn */
-	TArray<FVector> DrawPath;
-
-	/** Current Distance along Path */
-	float PathDistance = 0.0f;
-
-protected:
-
-	/** Current Desired Style */
-	FName DesiredStyle = NAME_None;
-
-	/** Current Desired Gait */
-	FName DesiredGait = NAME_None;
-
-	/** If walking is currently active */
-	bool bWalkActive = false;
-
-	/** Current Style Index */
-	int32 StyleIndex = 0;
-
-	/** Current desired speed */
-	float DesiredMinimumSpeed = 0.0f;
-	float DesiredMaximumSpeed = 0.0f;
-
-	/** Current change in desired speed */
-	float DesiredMinimumSpeedVelocity = 0.0f;
-	float DesiredMaximumSpeedVelocity = 0.0f;
-
-	/** Current Desired Rotation */
-	FRotator DesiredRotation = FRotator::ZeroRotator;
-
-	/** Current Desired Velocity */
-	FVector DesiredVelocity = FVector::ZeroVector;
-
-	/** Current Trajectory */
-	FAnimGenTrajectory Trajectory;
-
-	/** Current Transform Trajectory */
-	FTransformTrajectory TransformTrajectory;
-
-protected:
-
-	/** Camera Parameters */
-
-	float CameraYawVelocity = 0.0f;
-	float CameraYawAcceleration = 0.0f;
-	float CameraPitchVelocity = 0.0f;
-	float CameraPitchAcceleration = 0.0f;
-	float CameraDistanceDelta = 0.0f;
-	float CameraDistanceVelocity = 0.0f;
-	float CameraDistanceAcceleration = 0.0f;
-
-#endif
-
-#if WITH_EDITOR
-
-public:
-
-	virtual FAnimGenControlObjectElement MakeControl_Implementation(FAnimGenControlObject& InControlObject, const FAnimDatabasePoseState& InPoseState) const override;
-
-	virtual void OnPreviewReset_Implementation(const FAnimDatabasePoseState& InPoseState) override;
-
-	virtual void OnPreviewUpdate_Implementation(const FAnimDatabasePoseState& InPoseState, const float DeltaTime) override;
-
-	virtual void OnPreviewInputKeyPressed_Implementation(const FKey Key) override;
-
-	virtual void OnPreviewInputKeyReleased_Implementation(const FKey Key) override;
-
-	virtual void OnPreviewDrawDebug_Implementation(
-		const FDebugDrawer& Drawer,
-		const FAnimGenControlObject& InControlObject,
-		const FAnimGenControlObjectElement InControlObjectElement,
-		const FAnimDatabasePoseState& InPoseState,
-		const FLinearColor& IdentifierColor) override;
-
-
-	virtual void OnPreviewInputLeftMousePressed_Implementation(const FVector MouseGroundLocation) override;
-
-	virtual void OnPreviewInputLeftMouseReleased_Implementation(const FVector MouseGroundLocation) override;
-
-	virtual void OnPreviewInputMouseMoved_Implementation(const FVector MouseGroundLocation) override;
-#endif
-};
-
-
-
-/** Behavior preview for Move To Target */
-UCLASS(BlueprintType, Blueprintable, meta = (DisplayName = "MoveToTarget"))
-class UAnimGenBehaviorPreview_MoveToTarget : public UAnimGenBehaviorPreview
-{
-	GENERATED_BODY()
-
-#if WITH_EDITORONLY_DATA
-
-public:
-
-	/** List of styles to cycle through */
-	UPROPERTY(EditAnywhere, Category = "Styles")
-	TArray<FName> StyleTags;
-
-protected:
-
-	/** Current Desired Style */
-	FName DesiredStyle = NAME_None;
-
-	/** Current Style Index */
-	int32 StyleIndex = 0;
-
-	/** Current Desired Target Location */
-	FVector TargetLocation = FVector::ZeroVector;
-
-	/** Current Desired Target Direction */
-	FVector TargetDirection = FVector::RightVector;
-
-	/** If to use the current target direction */
-	bool bUseTargetDirection = true;
-
-	/** If the active location is being placed */
-	bool bPlacingActive = false;
-
-	/** Current Active Location */
-	FVector ActiveLocation = FVector::ZeroVector;
-
-	/** Current Active Direction */
-	FVector ActiveDirection = FVector::RightVector;
-
-	/** If to use the current active direction */
-	bool bUseActiveDirection = true;
-
-#endif
-
-#if WITH_EDITOR
-
-public:
-
-	virtual FAnimGenControlObjectElement MakeControl_Implementation(FAnimGenControlObject& InControlObject, const FAnimDatabasePoseState& InPoseState) const override;
-
-	virtual void OnPreviewInputKeyPressed_Implementation(const FKey Key) override;
-
-	virtual void OnPreviewDrawDebug_Implementation(
-		const FDebugDrawer& Drawer,
-		const FAnimGenControlObject& InControlObject,
-		const FAnimGenControlObjectElement InControlObjectElement,
-		const FAnimDatabasePoseState& InPoseState,
-		const FLinearColor& IdentifierColor) override;
-
-	virtual void OnPreviewInputLeftMousePressed_Implementation(const FVector MouseGroundLocation) override;
-
-	virtual void OnPreviewInputLeftMouseReleased_Implementation(const FVector MouseGroundLocation) override;
-
-	virtual void OnPreviewInputMouseMoved_Implementation(const FVector MouseGroundLocation) override;
-#endif
-};
-
 #undef UE_API
\ No newline at end of file
diff --git a/Engine/Plugins/Experimental/Animation/AnimGen/Source/AnimGen/Public/AnimGenController.h b/Engine/Plugins/Experimental/Animation/AnimGen/Source/AnimGen/Public/AnimGenController.h
--- a/Engine/Plugins/Experimental/Animation/AnimGen/Source/AnimGen/Public/AnimGenController.h
+++ b/Engine/Plugins/Experimental/Animation/AnimGen/Source/AnimGen/Public/AnimGenController.h
@@ -424,117 +424,6 @@ public:
 #endif
 };
 
-/** Class containing all the UAnimGenController Editor window testing settings */
-UCLASS()
-class UAnimGenControllerTestingSettings : public UObject
-{
-	GENERATED_BODY()
-
-public:
-
-#if WITH_EDITOR
-
-	/** Load the current database. */
-	UE_API void LoadDatabaseAsync();
-
-	/** We use this async load changes to ResetDatabase */
-	UE_API virtual void PostEditChangeProperty(struct FPropertyChangedEvent& Event) override;
-
-#endif
-
-#if WITH_EDITORONLY_DATA
-
-public:
-
-	/** Behavior Preview class to use */
-	UPROPERTY(EditAnywhere, Instanced, Category = "Scene")
-	TObjectPtr<UAnimGenBehaviorPreview> BehaviorPreview = nullptr;
-
-	/** Number of characters in the preview scene */
-	UPROPERTY(EditAnywhere, Category = "Scene", meta = (ClampMin = "1", ClampMax = "16", UIMin = "1", UIMax = "16"))
-	int32 CharacterNum = 1;
-
-	/** If to use a fixed timestep in the preview scene */
-	UPROPERTY(EditAnywhere, Category = "Scene")
-	bool bUseFixedTimeStep = false;
-
-	/** Fixed timestep to use in the preview scene */
-	UPROPERTY(EditAnywhere, Category = "Scene", meta = (EditCondition = "bUseFixedTimeStep", HideEditConditionToggle))
-	FFrameRate FixedTimeStep = FFrameRate(60, 1);
-
-	/** If to use a follow-cam in the preview scene */
-	UPROPERTY(EditAnywhere, Category = "Scene")
-	bool bFollowCam = false;
-
-public:
-
-	/** Database object to use for resetting the character in the scene */
-	UPROPERTY(EditAnywhere, Category = "Reset")
-	TSoftObjectPtr<UAnimDatabase> ResetDatabase = nullptr;
-
-	/** Set of frame ranges to use for resetting the character */
-	UPROPERTY(EditAnywhere, Instanced, Category = "Reset")
-	TObjectPtr<UAnimDatabaseFrameRangesFunction> ResetFrameRanges = nullptr;
-
-	/** Random seed to use when resetting the character */
-	UPROPERTY(EditAnywhere, Category = "Reset")
-	int32 ResetSeed = 53134;
-
-public:
-
-	/** The seed used for diffusion initial sampling, random seed is not used if set to 0 */
-	UPROPERTY(EditAnywhere, Category = "Inference", meta = (ClampMin = "0", UIMin = "0"))
-	int32 DiffusionSampleSeed = 1234;
-
-	/** Minimum Smoothing time to use - mixes the predicted pose with an integration of the predicted velocity. */
-	UPROPERTY(EditAnywhere, Category = "Inference", meta = (PinShownByDefault, ClampMin = "0.0", UIMin = "0.0"))
-	float MinPoseSmoothingTime = 0.05f;
-
-	/** Maximum Smoothing time to use - mixes the predicted pose with an integration of the predicted velocity. */
-	UPROPERTY(EditAnywhere, Category = "Inference", meta = (PinShownByDefault, ClampMin = "0.0", UIMin = "0.0"))
-	float MaxPoseSmoothingTime = 0.25f;
-
-	/** Amount of noise to use during sampling of motion */
-	UPROPERTY(EditAnywhere, Category = "Inference", meta = (ClampMin = "0.0", ClampMax = "1.0", UIMin = "0.0", UIMax = "1.0"))
-	float NoiseScale = 1.0f;
-
-	/** Which LOD level to use **/
-	UPROPERTY(EditAnywhere, Category = "Inference", meta = (ClampMin = "0", UIMin = "0", ClampMax = "2", UIMax = "2"))
-	int32 LODLevel = 0;
-
-public:
-
-	/** If to apply foot locking */
-	UPROPERTY(EditAnywhere, Category = "Inverse Kinematics")
-	bool bEnableFootLocking = false;
-
-	/** If to debug draw the foot locking */
-	UPROPERTY(EditAnywhere, Category = "Inverse Kinematics")
-	bool bDrawFootLocking = false;
-
-	/** Blend time to use for foot locking */
-	UPROPERTY(EditAnywhere, Category = "Inverse Kinematics", meta = (ForceUnits = "s"))
-	float FootLockingBlendTime = 0.25f;
-
-	/** Lock radius to use for foot locking */
-	UPROPERTY(EditAnywhere, Category = "Inverse Kinematics", meta = (ForceUnits = "cm"))
-	float FootLockingLockRadius = 1.0f;
-
-	/** Unlock radius to use for foot locking */
-	UPROPERTY(EditAnywhere, Category = "Inverse Kinematics", meta = (ForceUnits = "cm"))
-	float FootLockingUnlockRadius = 10.0f;
-
-	/** The length of the toe joint to use for foot locking */
-	UPROPERTY(EditAnywhere, Category = "Inverse Kinematics", meta = (ForceUnits = "cm"))
-	float FootLockingToeLength = 7.0f;
-
-	/** Amount of adjustment to apply to the pelvis height for foot locking */
-	UPROPERTY(EditAnywhere, Category = "Inverse Kinematics", meta=(ForceUnits = "cm"))
-	float FootLockingPelvisHeightAdjustment = -1.0f;
-
-#endif
-};
-
 /**
  * The Controller class represents a trained character controller that uses flow-matching and a neural-network to auto-regressively produce the 
  * character's pose frame by frame.
@@ -685,24 +574,12 @@ public:
 	UPROPERTY(VisibleAnywhere, Category = "Statistics", meta = (ForceUnits = "Kilobytes"))
 	int32 DenoiserSize = 0;
 
-	/** Current approximate inference time for the control encoder in microseconds */
-	UPROPERTY(VisibleAnywhere, Transient, Category = "Statistics", meta = (ForceUnits = "Microseconds"))
-	int32 ControlEncoderInferenceTime = 0;
-
-	/** Current approximate inference time for the denoiser in microseconds */
-	UPROPERTY(VisibleAnywhere, Transient, Category = "Statistics", meta = (ForceUnits = "Microseconds"))
-	int32 DenoiserInferenceTime = 0;
-
 public:
 
 	/** Editor-only object containing the training settings for the editor window */
 	UPROPERTY(Instanced)
 	TObjectPtr<UAnimGenControllerTrainingSettings> TrainingSettings;
 
-	/** Editor-only object containing the testing settings for the editor window */
-	UPROPERTY(Instanced)
-	TObjectPtr<UAnimGenControllerTestingSettings> TestingSettings;
-
 	/** Editor-only object containing the viewport settings for the editor window */
 	UPROPERTY(Instanced)
 	TObjectPtr<UAnimGenControllerViewportSettings> ViewportSettings;
diff --git a/Engine/Plugins/Experimental/Animation/AnimGen/Source/AnimGen/Public/RigVMFunction_AnimGen.h b/Engine/Plugins/Experimental/Animation/AnimGen/Source/AnimGen/Public/RigVMFunction_AnimGen.h
--- a/Engine/Plugins/Experimental/Animation/AnimGen/Source/AnimGen/Public/RigVMFunction_AnimGen.h
+++ b/Engine/Plugins/Experimental/Animation/AnimGen/Source/AnimGen/Public/RigVMFunction_AnimGen.h
@@ -180,10 +180,6 @@ struct FRigVMFunction_MakeTrajectoryFollowBehaviorControl : public FRigVMFunctio
 	UPROPERTY(meta = (Input))
 	FTransform RootBoneTransform;
 
-	// Desired Tags
-	UPROPERTY(meta = (Input))
-	TArray<FName> DesiredTags;
-
 	// Number of samples to take from the Trajectory
 	UPROPERTY(meta = (Input))
 	int32 TrajectorySampleNum = 4;
diff --git a/Engine/Plugins/Experimental/Animation/AnimGen/Source/AnimGenEditor/Private/AnimGenEditorAutoEncoderMode.cpp b/Engine/Plugins/Experimental/Animation/AnimGen/Source/AnimGenEditor/Private/AnimGenEditorAutoEncoderMode.cpp
--- a/Engine/Plugins/Experimental/Animation/AnimGen/Source/AnimGenEditor/Private/AnimGenEditorAutoEncoderMode.cpp
+++ b/Engine/Plugins/Experimental/Animation/AnimGen/Source/AnimGenEditor/Private/AnimGenEditorAutoEncoderMode.cpp
@@ -722,9 +722,6 @@ namespace UE::AnimGen::Editor
 				DecoderInferenceNetwork = AutoEncoder->DecoderNetwork;
 				EncoderInference = EncoderInferenceNetwork->GetNetwork()->CreateInferenceObject(MaxCharacterNum);
 				DecoderInference = DecoderInferenceNetwork->GetNetwork()->CreateInferenceObject(MaxCharacterNum);
-
-				AutoEncoder->EncoderSize = EncoderInferenceNetwork->GetSnapshotByteNum() / 1000;
-				AutoEncoder->DecoderSize = DecoderInferenceNetwork->GetSnapshotByteNum() / 1000;
 			}
 
 			if (EncodedVector.Num<1>() != EncoderInference->GetOutputSize() ||
diff --git a/Engine/Plugins/Experimental/Animation/AnimGen/Source/AnimGenEditor/Private/AnimGenEditorAutoEncoderToolkit.cpp b/Engine/Plugins/Experimental/Animation/AnimGen/Source/AnimGenEditor/Private/AnimGenEditorAutoEncoderToolkit.cpp
--- a/Engine/Plugins/Experimental/Animation/AnimGen/Source/AnimGenEditor/Private/AnimGenEditorAutoEncoderToolkit.cpp
+++ b/Engine/Plugins/Experimental/Animation/AnimGen/Source/AnimGenEditor/Private/AnimGenEditorAutoEncoderToolkit.cpp
@@ -1158,10 +1158,12 @@ namespace UE::AnimGen::Editor
 				const int32 EncoderByteNum = AutoEncoder->EncoderNetwork->GetSnapshotByteNum();
 				EncoderNetworkData = Learning::SharedMemory::Allocate<1, uint8>({ EncoderByteNum });
 				AutoEncoder->EncoderNetwork->SaveToSnapshot(MakeArrayView<uint8>(EncoderNetworkData.View.GetData(), EncoderNetworkData.View.Num()));
+				AutoEncoder->EncoderSize = AutoEncoder->EncoderNetwork->GetSnapshotByteNum() / 1000;
 
 				const int32 DecoderByteNum = AutoEncoder->DecoderNetwork->GetSnapshotByteNum();
 				DecoderNetworkData = Learning::SharedMemory::Allocate<1, uint8>({ DecoderByteNum });
 				AutoEncoder->DecoderNetwork->SaveToSnapshot(MakeArrayView<uint8>(DecoderNetworkData.View.GetData(), DecoderNetworkData.View.Num()));
+				AutoEncoder->DecoderSize = AutoEncoder->DecoderNetwork->GetSnapshotByteNum() / 1000;
 
 				// Write Config File
 
@@ -1636,7 +1638,7 @@ namespace UE::AnimGen::Editor
 
 	FText FAutoEncoderToolkit::GetBaseToolkitName() const
 	{
-		return LOCTEXT("AnimGenEditorAutoEncoderAppLabel", "Animation Auto Encoder Editor");
+		return LOCTEXT("AnimGenEditorAutoEncoderAppLabel", "AnimGen Auto-Encoder Editor");
 	}
 
 	FText FAutoEncoderToolkit::GetToolkitName() const
diff --git a/Engine/Plugins/Experimental/Animation/AnimGen/Source/AnimGenEditor/Private/AnimGenEditorControllerMode.cpp b/Engine/Plugins/Experimental/Animation/AnimGen/Source/AnimGenEditor/Private/AnimGenEditorControllerMode.cpp
--- a/Engine/Plugins/Experimental/Animation/AnimGen/Source/AnimGenEditor/Private/AnimGenEditorControllerMode.cpp
+++ b/Engine/Plugins/Experimental/Animation/AnimGen/Source/AnimGenEditor/Private/AnimGenEditorControllerMode.cpp
@@ -361,132 +361,11 @@ namespace UE::AnimGen::Editor
 	const FEditorModeID FControllerMode::EditorModeId = TEXT("AnimGenEditorControllerMode");
 
 	FControllerMode::FControllerMode() = default;
-	FControllerMode::~FControllerMode()
-	{
-		const int32 BehaviorNum = BehaviorPreviews.Num();
-		for (int32 BehaviorIdx = 0; BehaviorIdx < BehaviorNum; BehaviorIdx++)
-		{
-			if (BehaviorPreviews[BehaviorIdx] && BehaviorPreviews[BehaviorIdx]->IsValidLowLevel())
-			{
-				BehaviorPreviews[BehaviorIdx]->RemoveFromRoot();
-				BehaviorPreviews[BehaviorIdx] = nullptr;
-			}
-		}
-	}
-
-	bool FControllerMode::MouseMove(FEditorViewportClient* ViewportClient, FViewport* Viewport, int32 X, int32 Y)
-	{
-		if (!bCaptureInput)
-		{
-			return false;
-		}
-
-		FSceneViewFamilyContext ViewFamily(FSceneViewFamily::ConstructionValues(Viewport, ViewportClient->GetScene(), ViewportClient->EngineShowFlags).SetRealtimeUpdate(ViewportClient->IsRealtime()));
-		FSceneView* View = ViewportClient->CalcSceneView(&ViewFamily);
-		FViewportCursorLocation Cursor(View, ViewportClient, X, Y);
-
-		const FVector RayDirection = Cursor.GetDirection();
-
-		if (FMath::Abs(RayDirection.Z) < UE_SMALL_NUMBER)
-		{
-			return true;
-		}
-
-		const FVector RayStart = Cursor.GetOrigin();
-		const float Time = -RayStart.Z / RayDirection.Z;
-
-		if (Time > 0.0f)
-		{
-			const int32 BehaviorNum = BehaviorPreviews.Num();
-			for (int32 BehaviorIdx = 0; BehaviorIdx < BehaviorNum; BehaviorIdx++)
-			{
-				BehaviorPreviews[BehaviorIdx]->PreviewMouseGroundLocation = RayStart + Time * RayDirection;
-			}
-		}
-
-		return true;
-	}
-
-	bool FControllerMode::CapturedMouseMove(FEditorViewportClient* ViewportClient, FViewport* Viewport, int32 X, int32 Y)
-	{
-		return MouseMove(ViewportClient, Viewport, X, Y);
-	}
-
-	bool FControllerMode::InputAxis(FEditorViewportClient* InViewportClient, FViewport* Viewport, int32 ControllerId, FKey Key, float Delta, float DeltaTime)
-	{
-		if (!bCaptureInput)
-		{
-			return false;
-		}
-
-		const int32 BehaviorNum = BehaviorPreviews.Num();
-
-		if (Key == EKeys::Gamepad_LeftX)
-		{
-			for (int32 BehaviorIdx = 0; BehaviorIdx < BehaviorNum; BehaviorIdx++)
-			{
-				BehaviorPreviews[BehaviorIdx]->PreviewLeftStick.Y = Delta;
-			}
-			return true;
-		}
-		else if (Key == EKeys::Gamepad_LeftY)
-		{
-			for (int32 BehaviorIdx = 0; BehaviorIdx < BehaviorNum; BehaviorIdx++)
-			{
-				BehaviorPreviews[BehaviorIdx]->PreviewLeftStick.X = Delta;
-			}
-			return true;
-		}
-		else if (Key == EKeys::Gamepad_RightX)
-		{
-			for (int32 BehaviorIdx = 0; BehaviorIdx < BehaviorNum; BehaviorIdx++)
-			{
-				BehaviorPreviews[BehaviorIdx]->PreviewRightStick.Y = Delta;
-			}
-			return true;
-		}
-		else if (Key == EKeys::Gamepad_RightY)
-		{
-			for (int32 BehaviorIdx = 0; BehaviorIdx < BehaviorNum; BehaviorIdx++)
-			{
-				BehaviorPreviews[BehaviorIdx]->PreviewRightStick.X = -Delta;
-			}
-			return true;
-		}
-		else if (Key == EKeys::Gamepad_LeftTriggerAxis)
-		{
-			for (int32 BehaviorIdx = 0; BehaviorIdx < BehaviorNum; BehaviorIdx++)
-			{
-				BehaviorPreviews[BehaviorIdx]->PreviewLeftTrigger = Delta;
-			}
-			return true;
-		}
-		else if (Key == EKeys::Gamepad_RightTriggerAxis)
-		{
-			for (int32 BehaviorIdx = 0; BehaviorIdx < BehaviorNum; BehaviorIdx++)
-			{
-				BehaviorPreviews[BehaviorIdx]->PreviewRightTrigger = Delta;
-			}
-			return true;
-		}
-
-		return false;
-	}
-
-	bool FControllerMode::InputKey(FEditorViewportClient* ViewportClient, FViewport* Viewport, FKey Key, EInputEvent Event)
-	{
-		if (!bCaptureInput)
-		{
-			return false;
-		}
-
-		KeyEvents.Add({ Key, Event });
-		return true;
-	}
+	FControllerMode::~FControllerMode() = default;
 
-	void FControllerMode::TickTraining(FEditorViewportClient* EditorViewportClient, float DeltaTime)
+	void FControllerMode::Tick(FEditorViewportClient* EditorViewportClient, float DeltaTime)
 	{
-		bCaptureInput = false;
+		FEdMode::Tick(EditorViewportClient, DeltaTime);
 
 		AnimDatabase::Editor::FViewportClient* ViewportClient = static_cast<AnimDatabase::Editor::FViewportClient*>(EditorViewportClient);
 		check(ViewportClient);
@@ -548,11 +427,6 @@ namespace UE::AnimGen::Editor
 			return;
 		}
 
-		// Load Behavior
-
-		//UClass* TrainingBehavior = Controller->TrainingSettings->Behavior;
-		//if (!TrainingBehavior) { return; }
-
 		// Add the required number of characters to the scene
 
 		while (PreviewScene->GetCharacterNum() < MaxCharacterNum)
@@ -762,7 +636,7 @@ namespace UE::AnimGen::Editor
 			for (int32 AttributeIdx = 0; AttributeIdx < AttributeNum; AttributeIdx++)
 			{
 				FrameAttributes[AttributeIdx] = UAnimDatabaseFrameAttributeLibrary::MakeFrameAttributeFromFunction(
-					TrainingDatabase, 
+					TrainingDatabase,
 					Controller->TrainingSettings->QueryRanges,
 					FrameAttributeEntries[AttributeIdx].FrameAttribute);
 
@@ -886,1037 +760,165 @@ namespace UE::AnimGen::Editor
 		EditorToolKit->TimelineModel->UpdateTimelineFrameTime(DeltaTime * TrainingDatabase->GetFrameRate().AsDecimal());
 	}
 
-	void FControllerMode::TickTesting(FEditorViewportClient* EditorViewportClient, float DeltaTime)
+	void FControllerMode::Render(const FSceneView* View, FViewport* Viewport, FPrimitiveDrawInterface* PDI)
 	{
-		AnimDatabase::Editor::FViewportClient* ViewportClient = static_cast<AnimDatabase::Editor::FViewportClient*>(EditorViewportClient);
-		check(ViewportClient);
-		ViewportClient->ClearWarningMessage();
-
-		TSharedPtr<FControllerToolkit> EditorToolKit = StaticCastSharedPtr<FControllerToolkit>(ViewportClient->GetAssetEditorToolkit().Pin());
-		if (!EditorToolKit.IsValid()) { return; }
-
-		WeakToolkit = EditorToolKit.ToWeakPtr();
-
-		TSharedPtr<UE::AnimDatabase::Editor::FPreviewScene> PreviewScene = ViewportClient->GetPreviewScene().Pin();
-		if (!PreviewScene.IsValid()) { return; }
-
-		// Get the Controller
-
-		UAnimGenController* Controller = EditorToolKit->Controller;
-		check(Controller);
-
-		// Check we have a trained controller
-
-		if (!Controller->IsValid()) { return; }
-
-		// Check we have a preview class
-
-		TObjectPtr<UAnimGenBehaviorPreview> BehaviorPreview = Controller->TestingSettings->BehaviorPreview;
-
-		if (!BehaviorPreview)
+		if (TSharedPtr<FControllerToolkit> ToolKit = WeakToolkit.Pin())
 		{
-			return;
-		}
+			// Load Database
 
-		// We need to use a fixed timestep because the controller is only trained on animation data at 60fps
+			UAnimDatabase* TrainingDatabase = ToolKit->Controller->TrainingSettings->Database.Get();
+			if (!TrainingDatabase) { return; }
 
-		if (Controller->TestingSettings->bUseFixedTimeStep)
-		{
-			DeltaTime = 1.0f / FMath::Max(Controller->TestingSettings->FixedTimeStep.AsDecimal(), UE_SMALL_NUMBER);
-		}
+			// Load AutoEncoder
 
-		// Add the required number of characters to the scene
+			UAnimGenAutoEncoder* TrainingAutoEncoder = ToolKit->Controller->TrainingSettings->AutoEncoder.Get();
+			if (!TrainingAutoEncoder || !TrainingAutoEncoder->IsValid()) { return; }
 
-		while (PreviewScene->GetCharacterNum() < MaxCharacterNum)
-		{
-			const int32 CharacterIdx = PreviewScene->AddCharacter();
-			PreviewScene->SetCharacterVisibility(CharacterIdx, false);
-		}
+			// Load Behavior
 
-		// Init Character Colors
+			TObjectPtr<UAnimGenBehavior> TrainingBehavior = ToolKit->Controller->TrainingSettings->Behavior;
+			if (!TrainingBehavior) { return; }
 
-		if (CharacterColors.Num() != MaxCharacterNum)
-		{
-			CharacterColors.SetNumUninitialized(MaxCharacterNum);
+			// Draw Skeleton
 
-			const uint32 SeedHash = Learning::Random::Int(Controller->ViewportSettings->RangeIdentifierColorSeed);
-			for (int32 CharacterIdx = 0; CharacterIdx < MaxCharacterNum; CharacterIdx++)
+			if (ToolKit->Controller->ViewportSettings->bDrawSkeleton)
 			{
-				const uint32 ColorHash = Learning::Random::Int(0x94f67075 ^ CharacterIdx ^ SeedHash);
-				const uint8 Hue = ColorHash & 0x000000FF;
-				const uint8 Saturation = ((ColorHash & 0x0000FF00) >> 8) % 25 + 230;
-				const uint8 Brightness = ((ColorHash & 0x00FF0000) >> 16) % 50 + 205;
-				CharacterColors[CharacterIdx] = FLinearColor::MakeFromHSV8(Hue, Saturation, Brightness);
-			}
-		}
-
-		// Load the AutoEncoder 
-
-		UAnimGenAutoEncoder* AutoEncoder = Controller->AutoEncoder;
-		if (!AutoEncoder || !AutoEncoder->IsValid()) { return; }
-
-		//TSubclassOf<UAnimGenBehavior> BehaviorClass = Controller->TrainedBehavior;
-		//if (!BehaviorClass) { return; }
-
-		//const UAnimGenBehavior* Behavior = BehaviorClass->GetDefaultObject<UAnimGenBehavior>();
+				// Find bone indices to render
 
-		const int32 CurrentAutoEncoderContentHash = AutoEncoder->GetContentHash();
-		if (CurrentAutoEncoderContentHash != TrainedAutoEncoderContentHash)
-		{
-			TrainedAutoEncoderContentHash = CurrentAutoEncoderContentHash;
-			EditorToolKit->EditingAssetWidget->ForceRefresh();
-		}
+				const int32 BoneNum = TrainingAutoEncoder->GetBoneNum();
 
-		const int32 BoneNum = AutoEncoder->GetBoneNum();
-		const int32 PoseVectorSize = AutoEncoder->GetPoseVectorSize();
-		const int32 EncodedPoseVectorSize = AutoEncoder->GetEncodingSize();
+				TArray<int32> BoneIndices;
 
-		// TODO: Check Encoding Size Matches
+				if (ToolKit->Controller->ViewportSettings->bDrawRequiredBonesOnly)
+				{
+					BoneIndices = TrainingAutoEncoder->AutoEncodedRequiredBoneIndices;
+				}
+				else
+				{
+					BoneIndices.SetNumUninitialized(BoneNum);
+					for (int32 BoneIdx = 0; BoneIdx < BoneNum; BoneIdx++)
+					{
+						BoneIndices[BoneIdx] = BoneIdx;
+					}
+				}
 
-		// Set the character's skeletal mesh based on what is given in the animation database
+				const FLinearColor DefaultBoneColor = GetDefault<UPersonaOptions>()->DefaultBoneColor;
 
-		USkinnedAsset* DesiredSkeletalMesh = AutoEncoder->TrainedSkeleton ? AutoEncoder->TrainedSkeleton->GetPreviewMesh() : nullptr;
-		DesiredSkeletalMesh = Controller->ViewportSettings->PreviewMesh ? Controller->ViewportSettings->PreviewMesh.Get() : DesiredSkeletalMesh;
+				for (int32 CharacterIdx = 0; CharacterIdx < CharacterNum; CharacterIdx++)
+				{
+					// Draw Skeleton
 
-		for (int32 CharacterIdx = 0; CharacterIdx < MaxCharacterNum; CharacterIdx++)
-		{
-			USkinnedAsset* CurrentSkeletalMesh = PreviewScene->GetPoseableMeshComponent(CharacterIdx)->GetSkinnedAsset();
+					FLinearColor BoneColor = DefaultBoneColor;
+					BoneColor.A = ToolKit->Controller->ViewportSettings->DrawSkeletonOpacity;
 
-			if (DesiredSkeletalMesh && DesiredSkeletalMesh != CurrentSkeletalMesh)
-			{
-				PreviewScene->GetPoseableMeshComponent(CharacterIdx)->SetSkinnedAssetAndUpdate(DesiredSkeletalMesh, true);
-			}
-		}
+					FDrawDebugSkeletonSettings SkeletonSettings;
+					SkeletonSettings.bDrawSimpleSkeleton = ToolKit->Controller->ViewportSettings->bDrawSimpleSkeleton;
+					SkeletonSettings.BoneRadius = ToolKit->Controller->ViewportSettings->DrawSkeletonScale;
 
-		// Check if the trained database has changed and if so force refresh of details panel to check if warning needs to be displayed
+					UDrawDebugLibrary::DrawDebugSkeletonArrayView(
+						FDebugDrawer::MakeDebugDrawer(PDI),
+						PoseState.CurrPoseData->RootData.RootLocations[CharacterIdx],
+						PoseState.CurrPoseData->RootData.RootRotations[CharacterIdx],
+						PoseState.PoseGlobalBoneData->BoneLocations[CharacterIdx],
+						PoseState.PoseGlobalBoneData->BoneRotations[CharacterIdx],
+						BoneIndices,
+						ToolKit->Controller->AutoEncoder->GetBoneParents(),
+						BoneColor,
+						false,
+						SkeletonSettings);
 
-		/*
-		UAnimDatabase* TrainedDatabase = Controller->TrainedDatabase.LoadSynchronous();
-		if (!TrainedDatabase) { return; }
+					// Draw Velocities
 
-		const int32 CurrentTrainedDatabaseContentHash = TrainedDatabase->GetContentHash();
-		UObject* CurrentTrainedFrameRangesDefaultObject =
-			!Controller->TrainedFrameRanges ? nullptr :
-			Controller->TrainedFrameRanges->GetDefaultObject();
+					FLinearColor VelocityBoneColor = BoneColor;
+					VelocityBoneColor.A = ToolKit->Controller->ViewportSettings->DrawLinearVelocitiesOpacity;
 
-		if (CurrentTrainedDatabaseContentHash != TrainedDatabaseContentHash ||
-			CurrentTrainedFrameRangesDefaultObject != TrainedDatabaseFrameRangesDefaultObject)
-		{
-			TrainedDatabaseContentHash = CurrentTrainedDatabaseContentHash;
-			TrainedDatabaseFrameRangesDefaultObject = CurrentTrainedFrameRangesDefaultObject;
-			EditorToolKit->EditingAssetWidget->ForceRefresh();
-		}
-		*/
+					if (ToolKit->Controller->ViewportSettings->bDrawLinearVelocities)
+					{
+						UDrawDebugLibrary::DrawDebugBoneVelocitiesArrayView(
+							FDebugDrawer::MakeDebugDrawer(PDI),
+							PoseState.PoseGlobalBoneData->BoneLocations[CharacterIdx],
+							PoseState.PoseGlobalBoneData->BoneLinearVelocities[CharacterIdx],
+							BoneIndices,
+							VelocityBoneColor,
+							ToolKit->Controller->ViewportSettings->DrawSkeletonScale * 0.5f,
+							true,
+							ToolKit->Controller->ViewportSettings->DrawLinearVelocitiesScale);
+					}
 
-		// TODO: Check Database Skeleton Matches
+					if (ToolKit->Controller->ViewportSettings->bDrawAngularVelocities)
+					{
+						UDrawDebugLibrary::DrawDebugBoneVelocitiesArrayView(
+							FDebugDrawer::MakeDebugDrawer(PDI),
+							PoseState.PoseGlobalBoneData->BoneLocations[CharacterIdx],
+							PoseState.PoseGlobalBoneData->BoneAngularVelocities[CharacterIdx],
+							BoneIndices,
+							VelocityBoneColor,
+							ToolKit->Controller->ViewportSettings->DrawSkeletonScale * 0.5f,
+							true,
+							ToolKit->Controller->ViewportSettings->DrawAngularVelocitiesScale);
+					}
 
-		bReset = EditorToolKit->bResetRequested;
+					// Draw Transforms
 
-		// Check if the number of characters has changed
+					if (ToolKit->Controller->ViewportSettings->bDrawBoneTransforms)
+					{
+						UDrawDebugLibrary::DrawDebugTransformsArrayView(
+							FDebugDrawer::MakeDebugDrawer(PDI),
+							PoseState.PoseGlobalBoneData->BoneLocations[CharacterIdx],
+							PoseState.PoseGlobalBoneData->BoneRotations[CharacterIdx],
+							PoseState.PoseGlobalBoneData->BoneScales[CharacterIdx],
+							UDrawDebugLibrary::MakeDrawDebugLineStyleFromColor(BoneColor),
+							false,
+							ToolKit->Controller->ViewportSettings->DrawSkeletonScale * 5.0f);
+					}
+				}
+			}
 
-		if (CharacterNum != Controller->TestingSettings->CharacterNum)
-		{
-			CharacterNum = Controller->TestingSettings->CharacterNum;
+			// Draw Range Identifiers
 
-			for (int32 CharacterIdx = 0; CharacterIdx < MaxCharacterNum; CharacterIdx++)
+			if (ToolKit->Controller->ViewportSettings->bDrawRangeIdentifier)
 			{
-				PreviewScene->SetCharacterVisibility(CharacterIdx, CharacterIdx < CharacterNum);
-			}
-
-			bReset = true;
-		}
+				for (int32 CharacterIdx = 0; CharacterIdx < CharacterNum; CharacterIdx++)
+				{
+					FLinearColor IdentifierColor = ToolKit->Controller->TrainingSettings->QueryEntries[ActiveRanges[CharacterIdx]]->Color;
+					IdentifierColor.A = ToolKit->Controller->ViewportSettings->DrawRangeIdentifierOpacity;
 
-		if (bReset)
-		{
-			ActiveBehaviorObject = nullptr;
-		}
+					const FVector IdentifierLocation = (FVector)PoseState.CurrPoseData->RootData.RootLocations[CharacterIdx] +
+						(FVector)PoseState.CurrPoseData->RootData.RootRotations[CharacterIdx].RotateVector(FVector3f(0.0f, 0.0f, ToolKit->Controller->ViewportSettings->DrawRangeIdentifierHeight));
 
-		if (BehaviorPreview != ActiveBehaviorObject || BehaviorPreviews.Num() != CharacterNum)
-		{
-			ActiveBehaviorObject = BehaviorPreview;
+					FDrawDebugLineStyle LineStyle;
+					LineStyle.Color = IdentifierColor;
+					LineStyle.Thickness = ToolKit->Controller->ViewportSettings->DrawRangeIdentifierThickness;
 
-			const int32 BehaviorNum = BehaviorPreviews.Num();
-			for (int32 BehaviorIdx = 0; BehaviorIdx < BehaviorNum; BehaviorIdx++)
-			{
-				if (BehaviorPreviews[BehaviorIdx] && BehaviorPreviews[BehaviorIdx]->IsValidLowLevel())
-				{
-					BehaviorPreviews[BehaviorIdx]->RemoveFromRoot();
-					BehaviorPreviews[BehaviorIdx] = nullptr;
+					UDrawDebugLibrary::DrawDebugSimpleSphere(
+						FDebugDrawer::MakeDebugDrawer(PDI),
+						IdentifierLocation,
+						FRotator::ZeroRotator,
+						LineStyle,
+						true,
+						ToolKit->Controller->ViewportSettings->DrawRangeIdentifierRadius);
 				}
 			}
 
-			BehaviorPreviews.SetNumZeroed(CharacterNum);
+			// Draw Skeleton Roots
 
-			for (int32 CharacterIdx = 0; CharacterIdx < CharacterNum; CharacterIdx++)
+			if (ToolKit->Controller->ViewportSettings->bDrawRoot && CharacterNum > 0)
 			{
-				// Recreate the blueprint UAnimGenBehavior object. This will have the function we need for constructing the controls at the current frame
-				BehaviorPreviews[CharacterIdx] = DuplicateObject(BehaviorPreview, GetTransientPackage());
-				BehaviorPreviews[CharacterIdx]->AddToRoot();
-				BehaviorPreviews[CharacterIdx]->OnPreviewInitialize(Controller, AutoEncoder);
-			}
+				FDrawDebugLineStyle LineStyle;
+				LineStyle.Color = FLinearColor::White;
+				LineStyle.Color.A = ToolKit->Controller->ViewportSettings->DrawRootOpacity;
+				LineStyle.Thickness = ToolKit->Controller->ViewportSettings->DrawSkeletonScale;
 
-			ControlObject = UAnimGenControls::MakeControlObject();
+				UDrawDebugLibrary::DrawDebugRotationsQuatArrayView(
+					FDebugDrawer::MakeDebugDrawer(PDI),
+					PoseState.CurrPoseData->RootData.RootLocations,
+					PoseState.CurrPoseData->RootData.RootRotations,
+					LineStyle,
+					true,
+					ToolKit->Controller->ViewportSettings->DrawRootScale);
+			}
 
-			// Force refresh the UI as there is a chance the new behavior class is not compatible
-			EditorToolKit->EditingAssetWidget->ForceRefresh();
-
-			bReset = true;
-		}
-
-		// Create Networks
-
-		if (EncoderInferenceNetwork != AutoEncoder->EncoderNetwork ||
-			DecoderInferenceNetwork != AutoEncoder->DecoderNetwork)
-		{
-			EncoderInferenceNetwork = AutoEncoder->EncoderNetwork;
-			DecoderInferenceNetwork = AutoEncoder->DecoderNetwork;
-			EncoderInference = EncoderInferenceNetwork->GetNetwork()->CreateInferenceObject(MaxCharacterNum);
-			DecoderInference = DecoderInferenceNetwork->GetNetwork()->CreateInferenceObject(MaxCharacterNum);
-		}
-
-		if (ControlEncoderInferenceNetwork != Controller->ControlEncoderNetwork)
-		{
-			ControlEncoderInferenceNetwork = Controller->ControlEncoderNetwork;
-			ControlEncoderInference = ControlEncoderInferenceNetwork->GetNetwork()->CreateInferenceObject(MaxCharacterNum);
-
-			Controller->ControlEncoderSize = ControlEncoderInferenceNetwork->GetSnapshotByteNum() / 1000;
-		}
-
-		if (LOD0InferenceNetwork != Controller->LOD0Network && Controller->LOD0Network)
-		{
-			LOD0InferenceNetwork = Controller->LOD0Network;
-			LOD1InferenceNetwork = Controller->LOD1Network;
-			LOD2InferenceNetwork = Controller->LOD2Network;
-			LOD0Inference = LOD0InferenceNetwork->GetNetwork()->CreateInferenceObject(MaxCharacterNum);
-			LOD1Inference = LOD1InferenceNetwork->GetNetwork()->CreateInferenceObject(MaxCharacterNum);
-			LOD2Inference = LOD2InferenceNetwork->GetNetwork()->CreateInferenceObject(MaxCharacterNum);
-
-			Controller->DenoiserSize =
-				LOD0InferenceNetwork->GetSnapshotByteNum() / 1000 +
-				LOD1InferenceNetwork->GetSnapshotByteNum() / 1000 +
-				LOD2InferenceNetwork->GetSnapshotByteNum() / 1000;
-		}
-
-		// Get Reset Database
-
-		UAnimDatabase* ResetDatabase = Controller->TestingSettings->ResetDatabase.Get() ?
-			Controller->TestingSettings->ResetDatabase.Get() : Controller->TrainingSettings->Database.Get();
-
-		// Compute frame attributes we expect from the reset database
-
-		const int32 DatabaseContentHash = ResetDatabase->GetContentHash();
-
-		if (FrameAttributeEntries != AutoEncoder->TrainedFrameAttributes ||
-			FrameAttributesDatabaseContentHash != DatabaseContentHash)
-		{
-			FScopedSlowTask SlowTask(0.0f, LOCTEXT("RebuildingFrameAttributes", "Rebuilding AutoEncoder Frame Attributes in Database"));
-			SlowTask.MakeDialog();
-
-			FrameAttributeEntries = AutoEncoder->TrainedFrameAttributes;
-			FrameAttributesDatabaseContentHash = DatabaseContentHash;
-
-			const int32 AttributeNum = FrameAttributeEntries.Num();
-			FrameAttributes.SetNum(AttributeNum);
-			for (int32 AttributeIdx = 0; AttributeIdx < AttributeNum; AttributeIdx++)
-			{
-				FrameAttributes[AttributeIdx] = UAnimDatabaseFrameAttributeLibrary::MakeFrameAttributeFromFunction(
-					ResetDatabase, 
-					UAnimDatabaseFrameRangesLibrary::MakeFrameRangesFromDatabase(ResetDatabase),
-					FrameAttributeEntries[AttributeIdx].FrameAttribute);
-			}
-		}
-
-		// Check if we need to manually reset the character state
-
-		if (bReset)
-		{
-			EditorToolKit->bResetRequested = false;
-
-			// Reset the Random State
-
-			RandomState = 0xd5ca438d ^ Controller->TestingSettings->DiffusionSampleSeed;
-
-			// Reset to random frame from tag
-
-			FAnimDatabaseFrameRanges ResetFrameRanges;
-
-			if (Controller->TestingSettings->ResetFrameRanges)
-			{
-				ResetFrameRanges = UAnimDatabaseFrameRangesLibrary::MakeFrameRangesFromFunction(ResetDatabase, UAnimDatabaseFrameRangesLibrary::MakeFrameRangesFromDatabase(ResetDatabase), Controller->TestingSettings->ResetFrameRanges);
-
-				if (!ResetFrameRanges.IsValid())
-				{
-					ResetFrameRanges = UAnimDatabaseFrameRangesLibrary::MakeFrameRangesFromDatabase(ResetDatabase);
-				}
-			}
-			else
-			{
-				ResetFrameRanges = UAnimDatabaseFrameRangesLibrary::MakeFrameRangesFromDatabase(ResetDatabase);
-			}
-
-			if (!PoseState.IsValid())
-			{
-				PoseState = UAnimDatabasePoseStateLibrary::MakePoseState();
-			}
-
-			PoseState.CurrPoseData->Resize(CharacterNum, BoneNum, AutoEncoder->GetAttributeTypes());
-			PoseState.PoseGlobalBoneData->Resize(CharacterNum, BoneNum);
-			PoseVector.SetNumUninitialized({ CharacterNum, PoseVectorSize });
-			UnnormalizedEncodedPoseVector.SetNumUninitialized({ CharacterNum, EncodedPoseVectorSize });
-			EncodedPoseVector.SetNumUninitialized({ CharacterNum, EncodedPoseVectorSize });
-
-			for (int32 CharacterIdx = 0; CharacterIdx < CharacterNum; CharacterIdx++)
-			{
-				// Sample a random sequence and time in the given frame ranges
-				int32 ResetRandomState = 0xfee24eb0 ^ Controller->TestingSettings->ResetSeed;
-
-				int32 SequenceIdx = INDEX_NONE;
-				float SequenceTime = 0.0f;
-				UAnimGenBehaviorLibrary::MakeRandomSequenceIndexAndTimeFromFrameRanges(
-					SequenceIdx,
-					SequenceTime,
-					ResetRandomState,
-					ResetFrameRanges,
-					1.0f / FMath::Max(ResetDatabase->GetFrameRate().AsDecimal(), UE_SMALL_NUMBER));
-
-				// Sample the pose data at the given sequence and time
-				// TODO: Only sample valid bones
-				ResetDatabase->SamplePoseData(
-					PoseState.CurrPoseData->Slice(CharacterIdx, 1),
-					{ SequenceIdx },
-					{ SequenceTime },
-					FrameAttributes);
-
-				// Teleport the character to the origin
-				AnimDatabase::PoseData::Teleport(
-					PoseState.CurrPoseData->RootData.Slice(CharacterIdx, 1),
-					{ FVector::ZeroVector },
-					{ FQuat4f::Identity });
-			}
-
-			*PoseState.PrevPoseData = *PoseState.CurrPoseData;
-			PoseState.DeltaTime = DeltaTime;
-
-			// Compute global bone transforms using forward kinematics
-			AnimDatabase::PoseData::ForwardKinematics(
-				PoseState.PoseGlobalBoneData->View(),
-				PoseState.CurrPoseData->LocalBoneData.ConstView(),
-				PoseState.CurrPoseData->RootData.ConstView(),
-				AutoEncoder->GetBoneParents());
-
-			const int32 LeftToeIndex = AutoEncoder->FindBoneIndex("ball_l");
-			const int32 RightToeIndex = AutoEncoder->FindBoneIndex("ball_r");
-
-			if (LeftToeIndex != INDEX_NONE && RightToeIndex != INDEX_NONE)
-			{
-				ContactState.Resize(CharacterNum, 2);
-				AnimDatabase::PoseData::ContactPinning::ResetContactPinning(
-					ContactState.View(),
-					PoseState.PoseGlobalBoneData->ConstView(),
-					{ LeftToeIndex, RightToeIndex  });
-			}
-
-			// Call the user-implemented blueprint preview scene reset function
-			for (int32 CharacterIdx = 0; CharacterIdx < CharacterNum; CharacterIdx++)
-			{
-				PoseState.PoseIdx = CharacterIdx;
-				BehaviorPreviews[CharacterIdx]->OnPreviewReset(PoseState);
-			}
-
-			// Reset the current encoded pose vector using the auto-encoder EncoderNetwork on the current pose data
-
-			AutoEncoder->ToPoseVectors(PoseVector, PoseState.CurrPoseData->ConstView());
-			AutoEncoder->NormalizePoseVectors(PoseVector);
-			EncoderInference->Evaluate(UnnormalizedEncodedPoseVector, PoseVector);
-			Controller->NormalizeEncodedPoseVectors(EncodedPoseVector, UnnormalizedEncodedPoseVector);
-		}
-
-		// Check for Custom Reset
-
-		for (int32 CharacterIdx = 0; CharacterIdx < CharacterNum; CharacterIdx++)
-		{
-			FAnimGenBehaviorPreviewReset ResetData;
-			if (BehaviorPreviews[CharacterIdx]->OnPreviewCheckReset(ResetData, PoseState))
-			{
-				ResetDatabase->SamplePoseData(
-					PoseState.CurrPoseData->Slice(CharacterIdx, 1),
-					{ ResetData.ResetSequenceIdx },
-					{ ResetData.ResetSequenceTime },
-					FrameAttributes);
-
-				AnimDatabase::PoseData::Teleport(
-					PoseState.CurrPoseData->RootData.Slice(CharacterIdx, 1),
-					{ ResetData.ResetLocation },
-					{ (FQuat4f)ResetData.ResetRotation.Quaternion() });
-
-				AnimDatabase::PoseData::Copy(
-					PoseState.PrevPoseData->Slice(CharacterIdx, 1),
-					PoseState.CurrPoseData->ConstSlice(CharacterIdx, 1));
-				PoseState.DeltaTime = DeltaTime; // TODO: This is for all characters?
-
-				AnimDatabase::PoseData::ForwardKinematics(
-					PoseState.PoseGlobalBoneData->Slice(CharacterIdx, 1),
-					PoseState.CurrPoseData->LocalBoneData.ConstSlice(CharacterIdx, 1),
-					PoseState.CurrPoseData->RootData.ConstSlice(CharacterIdx, 1),
-					AutoEncoder->GetBoneParents());
-
-				const int32 LeftToeIndex = AutoEncoder->FindBoneIndex("ball_l");
-				const int32 RightToeIndex = AutoEncoder->FindBoneIndex("ball_r");
-
-				if (LeftToeIndex != INDEX_NONE && RightToeIndex != INDEX_NONE)
-				{
-					ContactState.Resize(CharacterNum, 2);
-					AnimDatabase::PoseData::ContactPinning::ResetContactPinning(
-						ContactState.View(),
-						PoseState.PoseGlobalBoneData->ConstView(),
-						{ LeftToeIndex, RightToeIndex });
-				}
-
-				// Call the user-implemented blueprint preview scene reset function
-				PoseState.PoseIdx = CharacterIdx;
-				BehaviorPreviews[CharacterIdx]->OnPreviewReset(PoseState);
-
-				// Reset the current encoded pose vector using the auto-encoder EncoderNetwork on the current pose data
-
-				AutoEncoder->ToPoseVectors(PoseVector.Slice(CharacterIdx, 1), PoseState.CurrPoseData->ConstSlice(CharacterIdx, 1));
-				AutoEncoder->NormalizePoseVectors(PoseVector.Slice(CharacterIdx, 1));
-				EncoderInference->Evaluate(UnnormalizedEncodedPoseVector.Slice(CharacterIdx, 1), PoseVector.Slice(CharacterIdx, 1));
-				Controller->NormalizeEncodedPoseVectors(EncodedPoseVector.Slice(CharacterIdx, 1), UnnormalizedEncodedPoseVector.Slice(CharacterIdx, 1));
-			}
-		}
-
-		// Are we playing?
-
-		if (EditorToolKit->ControllerState == EControllerState::Playing ||
-			EditorToolKit->ControllerState == EControllerState::Stepping)
-		{
-			// Capture input when playing (camera cannot move)
-
-			bCaptureInput = true;
-
-			// Evaluate Input
-
-			for (int32 CharacterIdx = 0; CharacterIdx < CharacterNum; CharacterIdx++)
-			{
-				if (BehaviorPreviews[CharacterIdx]->PreviewLeftStick != BehaviorPreviews[CharacterIdx]->PreviewPrevLeftStick)
-				{
-					BehaviorPreviews[CharacterIdx]->OnPreviewInputLeftStickAxis(BehaviorPreviews[CharacterIdx]->PreviewLeftStick);
-					BehaviorPreviews[CharacterIdx]->PreviewPrevLeftStick = BehaviorPreviews[CharacterIdx]->PreviewLeftStick;
-				}
-
-				if (BehaviorPreviews[CharacterIdx]->PreviewRightStick != BehaviorPreviews[CharacterIdx]->PreviewPrevRightStick)
-				{
-					BehaviorPreviews[CharacterIdx]->OnPreviewInputRightStickAxis(BehaviorPreviews[CharacterIdx]->PreviewRightStick);
-					BehaviorPreviews[CharacterIdx]->PreviewPrevRightStick = BehaviorPreviews[CharacterIdx]->PreviewRightStick;
-				}
-
-				if (BehaviorPreviews[CharacterIdx]->PreviewLeftTrigger != BehaviorPreviews[CharacterIdx]->PreviewPrevLeftTrigger)
-				{
-					BehaviorPreviews[CharacterIdx]->OnPreviewInputLeftTriggerAxis(BehaviorPreviews[CharacterIdx]->PreviewLeftTrigger);
-					BehaviorPreviews[CharacterIdx]->PreviewPrevLeftTrigger = BehaviorPreviews[CharacterIdx]->PreviewLeftTrigger;
-				}
-
-				if (BehaviorPreviews[CharacterIdx]->PreviewRightTrigger != BehaviorPreviews[CharacterIdx]->PreviewPrevRightTrigger)
-				{
-					BehaviorPreviews[CharacterIdx]->OnPreviewInputRightTriggerAxis(BehaviorPreviews[CharacterIdx]->PreviewRightTrigger);
-					BehaviorPreviews[CharacterIdx]->PreviewPrevRightTrigger = BehaviorPreviews[CharacterIdx]->PreviewRightTrigger;
-				}
-
-				if (BehaviorPreviews[CharacterIdx]->PreviewMouseGroundLocation != BehaviorPreviews[CharacterIdx]->PreviewPrevMouseGroundLocation)
-				{
-					BehaviorPreviews[CharacterIdx]->OnPreviewInputMouseMoved(BehaviorPreviews[CharacterIdx]->PreviewMouseGroundLocation);
-					BehaviorPreviews[CharacterIdx]->PreviewPrevMouseGroundLocation = BehaviorPreviews[CharacterIdx]->PreviewMouseGroundLocation;
-				}
-			}
-
-			while (KeyEvents.Num())
-			{
-				const TPair<FKey, EInputEvent> KeyEvent = KeyEvents.Pop();
-				const FKey Key = KeyEvent.Key;
-				const EInputEvent Event = KeyEvent.Value;
-
-				for (int32 CharacterIdx = 0; CharacterIdx < CharacterNum; CharacterIdx++)
-				{
-					if (Key == EKeys::LeftMouseButton && Event == EInputEvent::IE_Pressed)
-					{
-						BehaviorPreviews[CharacterIdx]->OnPreviewInputLeftMousePressed(BehaviorPreviews[CharacterIdx]->PreviewMouseGroundLocation);
-					}
-					else if (Key == EKeys::LeftMouseButton && Event == EInputEvent::IE_Released)
-					{
-						BehaviorPreviews[CharacterIdx]->OnPreviewInputLeftMouseReleased(BehaviorPreviews[CharacterIdx]->PreviewMouseGroundLocation);
-					}
-					else if (Key == EKeys::RightMouseButton && Event == EInputEvent::IE_Pressed)
-					{
-						BehaviorPreviews[CharacterIdx]->OnPreviewInputRightMousePressed(BehaviorPreviews[CharacterIdx]->PreviewMouseGroundLocation);
-					}
-					else if (Key == EKeys::RightMouseButton && Event == EInputEvent::IE_Released)
-					{
-						BehaviorPreviews[CharacterIdx]->OnPreviewInputRightMouseReleased(BehaviorPreviews[CharacterIdx]->PreviewMouseGroundLocation);
-					}
-					else if (Event == EInputEvent::IE_Pressed)
-					{
-						BehaviorPreviews[CharacterIdx]->OnPreviewInputKeyPressed(Key);
-					}
-					else if (Event == EInputEvent::IE_Released)
-					{
-						BehaviorPreviews[CharacterIdx]->OnPreviewInputKeyReleased(Key);
-					}
-				}
-			}
-
-			// Update Camera
-
-			if (Controller->TestingSettings->bFollowCam && CharacterNum > 0)
-			{
-				ViewportClient->SetViewRotation(FRotator(-BehaviorPreviews[0]->PreviewCameraPitch, BehaviorPreviews[0]->PreviewCameraYaw, 0.0f));
-				ViewportClient->SetViewLocationForOrbiting(
-					PoseState.CurrPoseData->RootData.RootLocations[0] + FVector(0.0f, 0.0f, BehaviorPreviews[0]->PreviewCameraHeight),
-					BehaviorPreviews[0]->PreviewCameraDistance);
-			}
-
-			for (int32 CharacterIdx = 0; CharacterIdx < CharacterNum; CharacterIdx++)
-			{
-				BehaviorPreviews[CharacterIdx]->PreviewCameraLocation = ViewportClient->GetViewLocation();
-				BehaviorPreviews[CharacterIdx]->PreviewCameraRotation = ViewportClient->GetViewRotation();
-			}
-
-			// Update Preview Scene
-
-			for (int32 CharacterIdx = 0; CharacterIdx < CharacterNum; CharacterIdx++)
-			{
-				PoseState.PoseIdx = CharacterIdx;
-				BehaviorPreviews[CharacterIdx]->OnPreviewUpdate(PoseState, DeltaTime);
-			}
-
-			// Record Previous Frame
-
-			*PoseState.PrevPoseData = *PoseState.CurrPoseData;
-			PoseState.DeltaTime = DeltaTime;
-
-			// Evaluate Controls. Here we call the blueprint implemented MakeControl function using the current character state.
-
-			const int32 ControlVectorSize = Controller->ControlSchema.ObservationSchema->Schema.GetObservationVectorSize(Controller->ControlSchemaElement.SchemaElement);
-			const int32 ControlDistributionVectorSize = Controller->ControlSchema.ObservationSchema->Schema.GetObservationDistributionVectorSize(Controller->ControlSchemaElement.SchemaElement);
-
-			ControlVector.SetNumUninitialized({ CharacterNum, ControlVectorSize });
-			ControlObject.ObservationObject->Object.Reset();
-			ControlObjectElements.SetNum(CharacterNum);
-
-			for (int32 CharacterIdx = 0; CharacterIdx < CharacterNum; CharacterIdx++)
-			{
-				PoseState.PoseIdx = CharacterIdx;
-				ControlObjectElements[CharacterIdx] = BehaviorPreviews[CharacterIdx]->MakeControl(ControlObject, PoseState);
-
-				// Validate that the return controls match the schema we provided
-
-				if (!UAnimGenControls::ValidateControlObjectMatchesSchema(
-					Controller->ControlSchema,
-					Controller->ControlSchemaElement,
-					ControlObject,
-					ControlObjectElements[CharacterIdx]))
-				{
-					return;
-				}
-
-				Learning::Observation::SetVectorFromObject(
-					ControlVector[CharacterIdx],
-					Controller->ControlSchema.ObservationSchema->Schema,
-					Controller->ControlSchemaElement.SchemaElement,
-					ControlObject.ObservationObject->Object,
-					ControlObjectElements[CharacterIdx].ObjectElement);
-			}
-
-			// Clamp Control Vectors
-
-			if (Controller->ControlVectorDistributionMins.Num() == ControlDistributionVectorSize &&
-				Controller->ControlVectorDistributionMaxs.Num() == ControlDistributionVectorSize)
-			{
-				for (int32 CharacterIdx = 0; CharacterIdx < CharacterNum; CharacterIdx++)
-				{
-					Learning::Observation::ClampToObservationDistributionMinMax(
-						ControlVector[CharacterIdx],
-						Controller->ControlVectorDistributionMins,
-						Controller->ControlVectorDistributionMaxs,
-						Controller->ControlSchema.ObservationSchema->Schema,
-						Controller->ControlSchemaElement.SchemaElement);
-				}
-			}
-
-			// Evaluate Control Encoder
-
-			const int32 EncodedControlVectorSize = Controller->EncodedControlVectorSize;
-			ControllerOutputVector.SetNumUninitialized({ CharacterNum, EncodedControlVectorSize });
-
-			if (!ControlEncoderInference)
-			{
-				return;
-			}
-
-			if (ControllerOutputVector.Num<1>() != ControlEncoderInference->GetOutputSize() ||
-				ControlVector.Num<1>() != ControlEncoderInference->GetInputSize())
-			{
-				ViewportClient->SetWarningMessage(LOCTEXT("EncoderNetworkSizeWarning", "Warning: Control Encoder Network sizes don't match. Is network asset out-of-date?"));
-				Learning::Array::Zero(ControllerOutputVector);
-			}
-			else
-			{
-				ControlEncoderPerfCounter.Begin();
-				if (EncodedControlVectorSize != 0) { ControlEncoderInference->Evaluate(ControllerOutputVector, ControlVector); }
-				ControlEncoderPerfCounter.End();
-
-				Controller->ControlEncoderInferenceTime = FMath::RoundToInt(FPlatformTime::ToMilliseconds64(ControlEncoderPerfCounter.GetAverage()) * 1000.0);
-			}
-
-			// Evaluate Denoiser
-
-			const float ControllerDeltaTime = 1.0f / FMath::Max(Controller->FrameRate.AsDecimal(), UE_SMALL_NUMBER);
-			const float SubsampleRate = FMath::Clamp(FMath::Max(DeltaTime, ControllerDeltaTime) / ControllerDeltaTime, 1, 8);
-
-			NextEncodedPoseVector = EncodedPoseVector;
-
-			DenoiserPerfCounter.Begin();
-
-			DenoiserInputVector.SetNumUninitialized({ CharacterNum, EncodedPoseVectorSize * 2 + 1 + EncodedControlVectorSize });
-
-			check(LOD0Inference->GetInputSize() == LOD1Inference->GetInputSize() && LOD0Inference->GetInputSize() == LOD2Inference->GetInputSize());
-			check(LOD0Inference->GetOutputSize() == LOD1Inference->GetOutputSize() && LOD0Inference->GetOutputSize() == LOD2Inference->GetOutputSize());
-
-			if (DenoiserInputVector.Num<1>() != LOD0Inference->GetInputSize() || EncodedPoseVectorSize != LOD0Inference->GetOutputSize())
-			{
-				ViewportClient->SetWarningMessage(LOCTEXT("DenoiserNetworkSizeWarning", "Warning: Denoiser Network sizes don't match. Is network asset out-of-date?"));
-				Learning::Array::Zero(NextEncodedPoseVector);
-			}
-			else
-			{
-				for (int32 CharacterIdx = 0; CharacterIdx < CharacterNum; CharacterIdx++)
-				{
-					Learning::Random::SampleClippedGaussianArray(NextEncodedPoseVector[CharacterIdx], RandomState, 0.0f, Controller->TestingSettings->NoiseScale, 3.0f);
-					if (!Controller->NormalizedPoseStds.IsEmpty())
-					{
-						for (int32 Idx = 0; Idx < EncodedPoseVectorSize; Idx++)
-						{
-							NextEncodedPoseVector[CharacterIdx][Idx] *= Controller->NormalizedPoseStds[Idx];
-						}
-					}
-
-					Learning::Array::Copy(DenoiserInputVector[CharacterIdx].Slice(0, EncodedPoseVectorSize), EncodedPoseVector[CharacterIdx]);
-					Learning::Array::Copy(DenoiserInputVector[CharacterIdx].Slice(EncodedPoseVectorSize, EncodedPoseVectorSize), NextEncodedPoseVector[CharacterIdx]);
-					Learning::Array::Set(DenoiserInputVector[CharacterIdx].Slice(EncodedPoseVectorSize * 2, 1), (SubsampleRate - 1.0f) / 7.0f);
-					Learning::Array::Copy(DenoiserInputVector[CharacterIdx].Slice(EncodedPoseVectorSize * 2 + 1, EncodedControlVectorSize), ControllerOutputVector[CharacterIdx]);
-				}
-
-				switch (Controller->TestingSettings->LODLevel)
-				{
-				default:
-				case 0: LOD0Inference->Evaluate(NextEncodedPoseVector, DenoiserInputVector); break;
-				case 1: LOD1Inference->Evaluate(NextEncodedPoseVector, DenoiserInputVector); break;
-				case 2: LOD2Inference->Evaluate(NextEncodedPoseVector, DenoiserInputVector); break;
-				}
-
-				Controller->ClampNormalizedPoseVectorsInplace(NextEncodedPoseVector);
-			}
-
-			DenoiserPerfCounter.End();
-
-			Controller->DenoiserInferenceTime = FMath::RoundToInt(FPlatformTime::ToMilliseconds64(DenoiserPerfCounter.GetAverage()) * 1000.0);
-
-			// Interpolate / Extrapolate using Delta Time Ratio
-
-			if (DeltaTime < ControllerDeltaTime)
-			{
-				AnimDatabase::Math::LerpToTargetInplace(
-					EncodedPoseVector,
-					NextEncodedPoseVector,
-					DeltaTime / ControllerDeltaTime);
-			}
-			else
-			{
-				Learning::Array::Copy(EncodedPoseVector, NextEncodedPoseVector);
-			}
-
-			Controller->ClampNormalizedPoseVectorsInplace(EncodedPoseVector);
-
-			// Evaluate Decoder
-
-			// Using the Decoder we can map from the EncodedPoseVector to the PoseVector, and then from the PoseVector we can get the PoseData
-
-			Controller->DenormalizeEncodedPoseVectors(UnnormalizedEncodedPoseVector, EncodedPoseVector);
-			PoseVector.SetNumUninitialized({ CharacterNum, PoseVectorSize });
-			DecoderInference->Evaluate(PoseVector, UnnormalizedEncodedPoseVector);
-
-			// Apply some smoothing based on mixing predicted velocities with predicted rotations and also account for DeltaTime variations
-
-			OutputPoseData.Resize(CharacterNum, BoneNum, AutoEncoder->GetAttributeTypes());
-
-			AutoEncoder->DenormalizePoseVectors(PoseVector);
-
-			AutoEncoder->FromPoseVectors(
-				OutputPoseData.View(),
-				PoseVector,
-				PoseState.CurrPoseData->RootData.RootLocations,
-				PoseState.CurrPoseData->RootData.RootRotations);
-
-			// Update Attributes
-			
-			PoseState.CurrPoseData->AttributeData = OutputPoseData.AttributeData;
-
-			// Update Pose
-
-			AnimDatabase::PoseData::DampInplaceUsingAdaptiveVelocityIntegrationMix(
-				PoseState.CurrPoseData->LocalBoneData.View(),
-				OutputPoseData.LocalBoneData.ConstView(),
-				DeltaTime,
-				Controller->TestingSettings->MinPoseSmoothingTime,
-				Controller->TestingSettings->MaxPoseSmoothingTime);
-
-			// Integrate Root
-
-			PoseState.CurrPoseData->RootData = OutputPoseData.RootData;
-			for (int32 CharacterIdx = 0; CharacterIdx < CharacterNum; CharacterIdx++)
-			{
-				PoseState.CurrPoseData->RootData.RootLocations[CharacterIdx] = DeltaTime * (FVector)OutputPoseData.RootData.RootLinearVelocities[CharacterIdx] + OutputPoseData.RootData.RootLocations[CharacterIdx];
-				PoseState.CurrPoseData->RootData.RootRotations[CharacterIdx] = FQuat4f::MakeFromRotationVector(DeltaTime * OutputPoseData.RootData.RootAngularVelocities[CharacterIdx]) * OutputPoseData.RootData.RootRotations[CharacterIdx];
-			}
-
-			PoseState.PoseGlobalBoneData->Resize(CharacterNum, BoneNum);
-			AnimDatabase::PoseData::ForwardKinematics(PoseState.PoseGlobalBoneData->View(), PoseState.CurrPoseData->LocalBoneData.ConstView(), PoseState.CurrPoseData->RootData.ConstView(), AutoEncoder->GetBoneParents());
-
-			AdjustedPoseData = *PoseState.CurrPoseData;
-
-			const int32 HipBoneIndex = AutoEncoder->FindBoneIndex("pelvis");
-
-			if (Controller->TestingSettings->bEnableFootLocking && HipBoneIndex != INDEX_NONE)
-			{
-				for (int32 CharacterIdx = 0; CharacterIdx < CharacterNum; CharacterIdx++)
-				{
-					AdjustedPoseData.LocalBoneData.BoneLocations[CharacterIdx][HipBoneIndex].Z += Controller->TestingSettings->FootLockingPelvisHeightAdjustment;
-				}
-			}
-
-			AdjustedGlobalBoneData.Resize(CharacterNum, BoneNum);
-			AnimDatabase::PoseData::ForwardKinematics(AdjustedGlobalBoneData.View(), AdjustedPoseData.LocalBoneData.ConstView(), AdjustedPoseData.RootData.ConstView(), AutoEncoder->BoneParents);
-
-			///////////////////
-
-			const int32 LeftContactAttribute = AutoEncoder->FindAttributeIndex(TEXT("contact_l"));
-			const int32 LeftLegBoneIndex = AutoEncoder->FindBoneIndex("thigh_l");
-			const int32 LeftKneeBoneIndex = AutoEncoder->FindBoneIndex("calf_l");
-			const int32 LeftAnkleBoneIndex = AutoEncoder->FindBoneIndex("foot_l");
-			const int32 LeftToeBoneIndex = AutoEncoder->FindBoneIndex("ball_l");
-
-			const int32 RightContactAttribute = AutoEncoder->FindAttributeIndex(TEXT("contact_r"));
-			const int32 RightLegBoneIndex = AutoEncoder->FindBoneIndex("thigh_r");
-			const int32 RightKneeBoneIndex = AutoEncoder->FindBoneIndex("calf_r");
-			const int32 RightAnkleBoneIndex = AutoEncoder->FindBoneIndex("foot_r");
-			const int32 RightToeBoneIndex = AutoEncoder->FindBoneIndex("ball_r");
-
-			if (Controller->TestingSettings->bEnableFootLocking &&
-				HipBoneIndex != INDEX_NONE &&
-				LeftContactAttribute != INDEX_NONE &&
-				LeftLegBoneIndex != INDEX_NONE &&
-				LeftKneeBoneIndex != INDEX_NONE &&
-				LeftAnkleBoneIndex != INDEX_NONE &&
-				LeftToeBoneIndex != INDEX_NONE &&
-				RightContactAttribute != INDEX_NONE &&
-				RightLegBoneIndex != INDEX_NONE &&
-				RightKneeBoneIndex != INDEX_NONE &&
-				RightAnkleBoneIndex != INDEX_NONE &&
-				RightToeBoneIndex != INDEX_NONE)
-			{
-				AnimDatabase::PoseData::ContactPinning::FContactPinningSettings ContactSettings;
-				ContactSettings.BlendTime = Controller->TestingSettings->FootLockingBlendTime;
-				ContactSettings.LockRadius = Controller->TestingSettings->FootLockingLockRadius;
-				ContactSettings.UnlockRadius = Controller->TestingSettings->FootLockingUnlockRadius;
-				ContactSettings.ContactHeight = AutoEncoder->TrainedSkeleton->GetReferenceSkeleton().GetBoneAbsoluteTransform(LeftToeBoneIndex).GetLocation().Z;
-				ContactSettings.MinToeHeight = AutoEncoder->TrainedSkeleton->GetReferenceSkeleton().GetBoneAbsoluteTransform(LeftToeBoneIndex).GetLocation().Z;
-				ContactSettings.MinAnkleHeight = AutoEncoder->TrainedSkeleton->GetReferenceSkeleton().GetBoneAbsoluteTransform(LeftAnkleBoneIndex).GetLocation().Z;
-				ContactSettings.ToeLength = Controller->TestingSettings->FootLockingToeLength;
-
-				AnimDatabase::PoseData::ContactPinning::UpdateContactPinning(
-					ContactState.View(),
-					AdjustedGlobalBoneData.View(),
-					AdjustedPoseData.LocalBoneData.View(),
-					AdjustedPoseData.RootData.ConstView(),
-					AdjustedPoseData.AttributeData.ConstView(),
-					AutoEncoder->GetBoneParents(),
-					{ LeftContactAttribute, RightContactAttribute },
-					{ HipBoneIndex, HipBoneIndex },
-					{ LeftLegBoneIndex, RightLegBoneIndex },
-					{ LeftKneeBoneIndex, RightKneeBoneIndex },
-					{ LeftAnkleBoneIndex, RightAnkleBoneIndex },
-					{ LeftToeBoneIndex, RightToeBoneIndex },
-					{ FVector3f(0.0f, 0.0f, +1.0f), FVector3f(0.0f, 0.0f, +1.0f) },
-					{ FVector3f(-1.0f, 0.0f, 0.0f), FVector3f(-1.0f, 0.0f, 0.0f) },
-					DeltaTime,
-					ContactSettings);
-			}
-
-			///////////////////
-
-			// Forward Kinematics to get the global transforms
-
-			AdjustedGlobalBoneData.Resize(CharacterNum, BoneNum);
-			AnimDatabase::PoseData::ForwardKinematics(AdjustedGlobalBoneData.View(), AdjustedPoseData.LocalBoneData.ConstView(), AdjustedPoseData.RootData.ConstView(), AutoEncoder->GetBoneParents());
-
-			///////////////////
-
-			// TODO: Why does there seem to be some error here where the root moves before the bones are updated?
-			//AnimGenViewportClient->GetPreviewScene()->GetCharacterActor()->SetActorTransform(RootTransform);
-
-			// Push the global transforms onto the character
-
-			for (int32 CharacterIdx = 0; CharacterIdx < CharacterNum; CharacterIdx++)
-			{
-				if (PreviewScene->GetPoseableMeshComponent(CharacterIdx)->GetSkinnedAsset())
-				{
-					for (int32 BoneIdx = 0; BoneIdx < BoneNum; BoneIdx++)
-					{
-						const FTransform BoneTransform = FTransform(
-							((FQuat)AdjustedGlobalBoneData.BoneRotations[CharacterIdx][BoneIdx]).GetNormalized(),
-							(FVector)AdjustedGlobalBoneData.BoneLocations[CharacterIdx][BoneIdx],
-							(FVector)AdjustedGlobalBoneData.BoneScales[CharacterIdx][BoneIdx]);
-
-						PreviewScene->GetPoseableMeshComponent(CharacterIdx)->SetBoneTransformByName(AutoEncoder->GetBoneName(BoneIdx), BoneTransform, EBoneSpaces::WorldSpace);
-					}
-
-					PreviewScene->GetPoseableMeshComponent(CharacterIdx)->RefreshBoneTransforms();
-				}
-			}
-
-			// Update UI Control Objects
-
-			if (ControlObject.ObservationObject.IsValid() && CharacterNum > 0)
-			{
-				TArray<TSharedPtr<FControlObjectTreeElement>>& ControlObjectTreeElements = EditorToolKit->ControlObjectWidget->GetTreeElementsRef();
-
-				bool bRefreshControlUI = ControlObjectTreeElements.Num() != CharacterNum;
-
-				ControlObjectTreeElements.SetNum(CharacterNum);
-
-				for (int32 CharacterIdx = 0; CharacterIdx < CharacterNum; CharacterIdx++)
-				{
-					if (!ControlObjectTreeElements[CharacterIdx].IsValid())
-					{
-						ControlObjectTreeElements[CharacterIdx] = MakeShared<FControlObjectTreeElement>();
-						bRefreshControlUI = true;
-					}
-
-					bRefreshControlUI |= Controller::Private::UpdateControlObjectTreeElement(
-						ControlObjectTreeElements[CharacterIdx],
-						ControlObject.ObservationObject->Object,
-						ControlObjectElements[CharacterIdx].ObjectElement);
-
-					ControlObjectTreeElements[CharacterIdx]->Color = CharacterColors[CharacterIdx];
-				}
-
-				if (bRefreshControlUI)
-				{
-					EditorToolKit->ControlObjectWidget->RefreshTreeElements();
-				}
-			}
-			else
-			{
-				TArray<TSharedPtr<FControlObjectTreeElement>>& ControlObjectTreeElements = EditorToolKit->ControlObjectWidget->GetTreeElementsRef();
-				if (!ControlObjectTreeElements.IsEmpty())
-				{
-					ControlObjectTreeElements.Empty();
-					EditorToolKit->ControlObjectWidget->RefreshTreeElements();
-				}
-			}
-
-			if (EditorToolKit->ControllerState == EControllerState::Stepping)
-			{
-				EditorToolKit->ControllerState = EControllerState::Paused;
-			}
-		}
-		else
-		{
-			bCaptureInput = false;
-		}
-	}
-
-	void FControllerMode::Tick(FEditorViewportClient* EditorViewportClient, float DeltaTime)
-	{
-		FEdMode::Tick(EditorViewportClient, DeltaTime);
-
-		AnimDatabase::Editor::FViewportClient* ViewportClient = static_cast<AnimDatabase::Editor::FViewportClient*>(EditorViewportClient);
-		check(ViewportClient);
-
-		TSharedPtr<FControllerToolkit> EditorToolKit = StaticCastSharedPtr<FControllerToolkit>(ViewportClient->GetAssetEditorToolkit().Pin());
-		if (!EditorToolKit.IsValid()) { return; }
-
-		WeakToolkit = EditorToolKit.ToWeakPtr();
-
-		if (EditorToolKit->ControllerState == EControllerState::Stopped)
-		{
-			TickTraining(EditorViewportClient, DeltaTime);
-			bReset = true;
-		}
-		else
-		{
-			TickTesting(EditorViewportClient, DeltaTime);
-		}
-	}
-
-	void FControllerMode::RenderTraining(const FSceneView* View, FViewport* Viewport, FPrimitiveDrawInterface* PDI)
-	{
-		if (TSharedPtr<FControllerToolkit> ToolKit = WeakToolkit.Pin())
-		{
-			// Load Database
-
-			UAnimDatabase* TrainingDatabase = ToolKit->Controller->TrainingSettings->Database.Get();
-			if (!TrainingDatabase) { return; }
-
-			// Load AutoEncoder
-
-			UAnimGenAutoEncoder* TrainingAutoEncoder = ToolKit->Controller->TrainingSettings->AutoEncoder.Get();
-			if (!TrainingAutoEncoder || !TrainingAutoEncoder->IsValid()) { return; }
-
-			// Load Behavior
-
-			TObjectPtr<UAnimGenBehavior> TrainingBehavior = ToolKit->Controller->TrainingSettings->Behavior;
-			if (!TrainingBehavior) { return; }
-
-			// Draw Skeleton
-
-			if (ToolKit->Controller->ViewportSettings->bDrawSkeleton)
-			{
-				// Find bone indices to render
-
-				const int32 BoneNum = TrainingAutoEncoder->GetBoneNum();
-
-				TArray<int32> BoneIndices;
-
-				if (ToolKit->Controller->ViewportSettings->bDrawRequiredBonesOnly)
-				{
-					BoneIndices = TrainingAutoEncoder->AutoEncodedRequiredBoneIndices;
-				}
-				else
-				{
-					BoneIndices.SetNumUninitialized(BoneNum);
-					for (int32 BoneIdx = 0; BoneIdx < BoneNum; BoneIdx++)
-					{
-						BoneIndices[BoneIdx] = BoneIdx;
-					}
-				}
-
-				const FLinearColor DefaultBoneColor = GetDefault<UPersonaOptions>()->DefaultBoneColor;
-
-				for (int32 CharacterIdx = 0; CharacterIdx < CharacterNum; CharacterIdx++)
-				{
-					// Draw Skeleton
-
-					FLinearColor BoneColor = DefaultBoneColor;
-					BoneColor.A = ToolKit->Controller->ViewportSettings->DrawSkeletonOpacity;
-
-					FDrawDebugSkeletonSettings SkeletonSettings;
-					SkeletonSettings.bDrawSimpleSkeleton = ToolKit->Controller->ViewportSettings->bDrawSimpleSkeleton;
-					SkeletonSettings.BoneRadius = ToolKit->Controller->ViewportSettings->DrawSkeletonScale;
-
-					UDrawDebugLibrary::DrawDebugSkeletonArrayView(
-						FDebugDrawer::MakeDebugDrawer(PDI),
-						PoseState.CurrPoseData->RootData.RootLocations[CharacterIdx],
-						PoseState.CurrPoseData->RootData.RootRotations[CharacterIdx],
-						PoseState.PoseGlobalBoneData->BoneLocations[CharacterIdx],
-						PoseState.PoseGlobalBoneData->BoneRotations[CharacterIdx],
-						BoneIndices,
-						ToolKit->Controller->AutoEncoder->GetBoneParents(),
-						BoneColor,
-						false,
-						SkeletonSettings);
-
-					// Draw Velocities
-
-					FLinearColor VelocityBoneColor = BoneColor;
-					VelocityBoneColor.A = ToolKit->Controller->ViewportSettings->DrawLinearVelocitiesOpacity;
-
-					if (ToolKit->Controller->ViewportSettings->bDrawLinearVelocities)
-					{
-						UDrawDebugLibrary::DrawDebugBoneVelocitiesArrayView(
-							FDebugDrawer::MakeDebugDrawer(PDI),
-							PoseState.PoseGlobalBoneData->BoneLocations[CharacterIdx],
-							PoseState.PoseGlobalBoneData->BoneLinearVelocities[CharacterIdx],
-							BoneIndices,
-							VelocityBoneColor,
-							ToolKit->Controller->ViewportSettings->DrawSkeletonScale * 0.5f,
-							true,
-							ToolKit->Controller->ViewportSettings->DrawLinearVelocitiesScale);
-					}
-
-					if (ToolKit->Controller->ViewportSettings->bDrawAngularVelocities)
-					{
-						UDrawDebugLibrary::DrawDebugBoneVelocitiesArrayView(
-							FDebugDrawer::MakeDebugDrawer(PDI),
-							PoseState.PoseGlobalBoneData->BoneLocations[CharacterIdx],
-							PoseState.PoseGlobalBoneData->BoneAngularVelocities[CharacterIdx],
-							BoneIndices,
-							VelocityBoneColor,
-							ToolKit->Controller->ViewportSettings->DrawSkeletonScale * 0.5f,
-							true,
-							ToolKit->Controller->ViewportSettings->DrawAngularVelocitiesScale);
-					}
-
-					// Draw Transforms
-
-					if (ToolKit->Controller->ViewportSettings->bDrawBoneTransforms)
-					{
-						UDrawDebugLibrary::DrawDebugTransformsArrayView(
-							FDebugDrawer::MakeDebugDrawer(PDI),
-							PoseState.PoseGlobalBoneData->BoneLocations[CharacterIdx],
-							PoseState.PoseGlobalBoneData->BoneRotations[CharacterIdx],
-							PoseState.PoseGlobalBoneData->BoneScales[CharacterIdx],
-							UDrawDebugLibrary::MakeDrawDebugLineStyleFromColor(BoneColor),
-							false,
-							ToolKit->Controller->ViewportSettings->DrawSkeletonScale * 5.0f);
-					}
-				}
-			}
-
-			// Draw Range Identifiers
-
-			if (ToolKit->Controller->ViewportSettings->bDrawRangeIdentifier)
-			{
-				for (int32 CharacterIdx = 0; CharacterIdx < CharacterNum; CharacterIdx++)
-				{
-					FLinearColor IdentifierColor = ToolKit->Controller->TrainingSettings->QueryEntries[ActiveRanges[CharacterIdx]]->Color;
-					IdentifierColor.A = ToolKit->Controller->ViewportSettings->DrawRangeIdentifierOpacity;
-
-					const FVector IdentifierLocation = (FVector)PoseState.CurrPoseData->RootData.RootLocations[CharacterIdx] +
-						(FVector)PoseState.CurrPoseData->RootData.RootRotations[CharacterIdx].RotateVector(FVector3f(0.0f, 0.0f, ToolKit->Controller->ViewportSettings->DrawRangeIdentifierHeight));
-
-					FDrawDebugLineStyle LineStyle;
-					LineStyle.Color = IdentifierColor;
-					LineStyle.Thickness = ToolKit->Controller->ViewportSettings->DrawRangeIdentifierThickness;
-
-					UDrawDebugLibrary::DrawDebugSimpleSphere(
-						FDebugDrawer::MakeDebugDrawer(PDI),
-						IdentifierLocation,
-						FRotator::ZeroRotator,
-						LineStyle,
-						true,
-						ToolKit->Controller->ViewportSettings->DrawRangeIdentifierRadius);
-				}
-			}
-
-			// Draw Skeleton Roots
-
-			if (ToolKit->Controller->ViewportSettings->bDrawRoot && CharacterNum > 0)
-			{
-				FDrawDebugLineStyle LineStyle;
-				LineStyle.Color = FLinearColor::White;
-				LineStyle.Color.A = ToolKit->Controller->ViewportSettings->DrawRootOpacity;
-				LineStyle.Thickness = ToolKit->Controller->ViewportSettings->DrawSkeletonScale;
-
-				UDrawDebugLibrary::DrawDebugRotationsQuatArrayView(
-					FDebugDrawer::MakeDebugDrawer(PDI),
-					PoseState.CurrPoseData->RootData.RootLocations,
-					PoseState.CurrPoseData->RootData.RootRotations,
-					LineStyle,
-					true,
-					ToolKit->Controller->ViewportSettings->DrawRootScale);
-			}
-
-			// Draw Origin
+			// Draw Origin
 
 			if (ToolKit->Controller->ViewportSettings->bDrawOrigin)
 			{
@@ -1996,268 +998,6 @@ namespace UE::AnimGen::Editor
 				}
 			}
 		}
-	}
-
-	void FControllerMode::RenderTesting(const FSceneView* View, FViewport* Viewport, FPrimitiveDrawInterface* PDI)
-	{
-		if (TSharedPtr<FControllerToolkit> ToolKit = WeakToolkit.Pin())
-		{
-			if (!ToolKit->Controller->IsValid()) { return; }
-			if (!ToolKit->Controller->AutoEncoder->IsValid()) { return; }
-			if (!ToolKit->Controller->TestingSettings->BehaviorPreview) { return; }
-
-			// Draw Skeleton
-
-			if (ToolKit->Controller->ViewportSettings->bDrawSkeleton)
-			{
-				// Find bone indices to render
-
-				const int32 BoneNum = ToolKit->Controller->AutoEncoder->GetBoneNum();
-
-				TArray<int32> BoneIndices;
-
-				if (ToolKit->Controller->ViewportSettings->bDrawRequiredBonesOnly)
-				{
-					BoneIndices = ToolKit->Controller->AutoEncoder->AutoEncodedRequiredBoneIndices;
-				}
-				else
-				{
-					BoneIndices.SetNumUninitialized(BoneNum);
-					for (int32 BoneIdx = 0; BoneIdx < BoneNum; BoneIdx++)
-					{
-						BoneIndices[BoneIdx] = BoneIdx;
-					}
-				}
-
-				const FLinearColor DefaultBoneColor = GetDefault<UPersonaOptions>()->DefaultBoneColor;
-
-				for (int32 CharacterIdx = 0; CharacterIdx < CharacterNum; CharacterIdx++)
-				{
-					// Draw Skeleton
-
-					FLinearColor BoneColor = DefaultBoneColor;
-					BoneColor.A = ToolKit->Controller->ViewportSettings->DrawSkeletonOpacity;
-
-					FDrawDebugSkeletonSettings SkeletonSettings;
-					SkeletonSettings.bDrawSimpleSkeleton = ToolKit->Controller->ViewportSettings->bDrawSimpleSkeleton;
-					SkeletonSettings.BoneRadius = ToolKit->Controller->ViewportSettings->DrawSkeletonScale;
-
-					UDrawDebugLibrary::DrawDebugSkeletonArrayView(
-						FDebugDrawer::MakeDebugDrawer(PDI),
-						PoseState.CurrPoseData->RootData.RootLocations[CharacterIdx],
-						PoseState.CurrPoseData->RootData.RootRotations[CharacterIdx],
-						PoseState.PoseGlobalBoneData->BoneLocations[CharacterIdx],
-						PoseState.PoseGlobalBoneData->BoneRotations[CharacterIdx],
-						BoneIndices,
-						ToolKit->Controller->AutoEncoder->GetBoneParents(),
-						BoneColor,
-						false,
-						SkeletonSettings);
-
-					// Draw Velocities
-
-					FLinearColor VelocityBoneColor = BoneColor;
-					VelocityBoneColor.A = ToolKit->Controller->ViewportSettings->DrawLinearVelocitiesOpacity;
-
-					if (ToolKit->Controller->ViewportSettings->bDrawLinearVelocities)
-					{
-						UDrawDebugLibrary::DrawDebugBoneVelocitiesArrayView(
-							FDebugDrawer::MakeDebugDrawer(PDI),
-							PoseState.PoseGlobalBoneData->BoneLocations[CharacterIdx],
-							PoseState.PoseGlobalBoneData->BoneLinearVelocities[CharacterIdx],
-							BoneIndices,
-							VelocityBoneColor,
-							ToolKit->Controller->ViewportSettings->DrawSkeletonScale * 0.5f,
-							false,
-							ToolKit->Controller->ViewportSettings->DrawLinearVelocitiesScale);
-					}
-
-					if (ToolKit->Controller->ViewportSettings->bDrawAngularVelocities)
-					{
-						UDrawDebugLibrary::DrawDebugBoneVelocitiesArrayView(
-							FDebugDrawer::MakeDebugDrawer(PDI),
-							PoseState.PoseGlobalBoneData->BoneLocations[CharacterIdx],
-							PoseState.PoseGlobalBoneData->BoneAngularVelocities[CharacterIdx],
-							BoneIndices,
-							VelocityBoneColor,
-							ToolKit->Controller->ViewportSettings->DrawSkeletonScale * 0.5f,
-							false,
-							ToolKit->Controller->ViewportSettings->DrawAngularVelocitiesScale);
-					}
-
-					// Draw Transforms
-
-					if (ToolKit->Controller->ViewportSettings->bDrawBoneTransforms)
-					{
-						UDrawDebugLibrary::DrawDebugTransformsArrayView(
-							FDebugDrawer::MakeDebugDrawer(PDI),
-							PoseState.PoseGlobalBoneData->BoneLocations[CharacterIdx],
-							PoseState.PoseGlobalBoneData->BoneRotations[CharacterIdx],
-							PoseState.PoseGlobalBoneData->BoneScales[CharacterIdx],
-							UDrawDebugLibrary::MakeDrawDebugLineStyleFromColor(BoneColor),
-							false,
-							ToolKit->Controller->ViewportSettings->DrawSkeletonScale * 5.0f);
-					}
-				}
-			}
-
-			// Draw Range Identifiers
-
-			if (ToolKit->Controller->ViewportSettings->bDrawRangeIdentifier)
-			{
-				for (int32 CharacterIdx = 0; CharacterIdx < CharacterNum; CharacterIdx++)
-				{
-					FLinearColor IdentifierColor = CharacterColors[CharacterIdx];
-					IdentifierColor.A = ToolKit->Controller->ViewportSettings->DrawRangeIdentifierOpacity;
-
-					const FVector IdentifierLocation = (FVector)PoseState.CurrPoseData->RootData.RootLocations[CharacterIdx] +
-						(FVector)PoseState.CurrPoseData->RootData.RootRotations[CharacterIdx].RotateVector(FVector3f(0.0f, 0.0f, ToolKit->Controller->ViewportSettings->DrawRangeIdentifierHeight));
-
-					FDrawDebugLineStyle LineStyle;
-					LineStyle.Color = IdentifierColor;
-					LineStyle.Thickness = ToolKit->Controller->ViewportSettings->DrawRangeIdentifierThickness;
-
-					UDrawDebugLibrary::DrawDebugSimpleSphere(
-						FDebugDrawer::MakeDebugDrawer(PDI),
-						IdentifierLocation,
-						FRotator::ZeroRotator,
-						LineStyle,
-						true,
-						ToolKit->Controller->ViewportSettings->DrawRangeIdentifierRadius);
-				}
-			}
-
-			// Draw Control Encoding
-
-			if (ToolKit->Controller->ViewportSettings->bDrawEncodedControlVector)
-			{
-				const int32 EncodedControlVectorSize = ControllerOutputVector.Num<1>();
-
-				const int32 DimNum = FMath::Min(EncodedControlVectorSize, ToolKit->Controller->ViewportSettings->DrawEncodedControlVectorDimensions);
-				const float VerticalScale = ToolKit->Controller->ViewportSettings->DrawEncodedControlVectorScale;
-				const FLinearColor DrawColor = ToolKit->Controller->ViewportSettings->DrawEncodedControlVectorColor;
-
-				FDrawDebugLineStyle LineStyle;
-				LineStyle.Color = DrawColor;
-				LineStyle.Thickness = 2.0f;
-
-				if (DimNum > 0 && ControllerOutputVector.Num<0>() > 0)
-				{
-					for (int32 CharacterIdx = 0; CharacterIdx < CharacterNum; CharacterIdx++)
-					{
-						const FVector Start = PoseState.CurrPoseData->RootData.RootLocations[CharacterIdx] + FVector(50.0f, 0.0f, 100.0f);
-						const FVector Offset = FVector(20.0f, 0.0f, 0.0f);
-
-						UDrawDebugLibrary::DrawDebugLine(FDebugDrawer::MakeDebugDrawer(PDI), Start, Start + Offset * (DimNum - 1), LineStyle);
-
-						for (int32 DimIdx = 0; DimIdx < DimNum; DimIdx++)
-						{
-							const float Value = ControllerOutputVector[CharacterIdx][DimIdx];
-							const FVector DimStart = Start + Offset * DimIdx;
-							const FVector DimEnd = Start + Offset * DimIdx + FVector(0.0f, 0.0f, 10.0f * VerticalScale * Value);
-
-							FDrawDebugStringSettings StringSettings;
-							StringSettings.Height = 6.0f;
-
-							UDrawDebugLibrary::DrawDebugLine(FDebugDrawer::MakeDebugDrawer(PDI), DimStart, DimEnd, LineStyle);
-							UDrawDebugLibrary::DrawDebugSphere(FDebugDrawer::MakeDebugDrawer(PDI), DimEnd, FRotator::ZeroRotator, LineStyle, true, 5.0f);
-							UDrawDebugLibrary::DrawDebugString(FDebugDrawer::MakeDebugDrawer(PDI), FString::Printf(TEXT("% 5.2f"), Value), DimEnd + FVector(-8.0f, 0.0f, 15.0f), FRotator::ZeroRotator, UDrawDebugLibrary::DrawDebugLineStyleWithThickness(LineStyle, 0.0), true, StringSettings);
-						}
-					}
-				}
-			}
-
-			// Draw Skeleton Roots
-
-			if (ToolKit->Controller->ViewportSettings->bDrawRoot && CharacterNum > 0)
-			{
-				FDrawDebugLineStyle LineStyle;
-				LineStyle.Color = FLinearColor::White;
-				LineStyle.Color.A = ToolKit->Controller->ViewportSettings->DrawRootOpacity;
-				LineStyle.Thickness = ToolKit->Controller->ViewportSettings->DrawSkeletonScale;
-
-				UDrawDebugLibrary::DrawDebugRotationsQuatArrayView(
-					FDebugDrawer::MakeDebugDrawer(PDI),
-					PoseState.CurrPoseData->RootData.RootLocations,
-					PoseState.CurrPoseData->RootData.RootRotations,
-					LineStyle,
-					true,
-					ToolKit->Controller->ViewportSettings->DrawRootScale);
-			}
-
-			// Draw Origin
-
-			if (ToolKit->Controller->ViewportSettings->bDrawOrigin)
-			{
-				FDrawDebugLineStyle LineStyle;
-				LineStyle.Color = FLinearColor::White;
-				LineStyle.Color.A = ToolKit->Controller->ViewportSettings->DrawOriginOpacity;
-				LineStyle.Thickness = ToolKit->Controller->ViewportSettings->DrawOriginLineThickness;
-
-				UDrawDebugLibrary::DrawDebugTransform(
-					FDebugDrawer::MakeDebugDrawer(PDI),
-					FTransform::Identity,
-					LineStyle,
-					true,
-					ToolKit->Controller->ViewportSettings->DrawOriginScale);
-			}
-
-			// Debug Draw IK
-
-			if (ToolKit->Controller->TestingSettings->bDrawFootLocking)
-			{
-				for (int32 CharacterIdx = 0; CharacterIdx < CharacterNum; CharacterIdx++)
-				{
-					for (int32 BoneIdx = 0; BoneIdx < 2; BoneIdx++)
-					{
-						if (ContactState.Locked[CharacterIdx][BoneIdx])
-						{
-							FDrawDebugLineStyle LineStyle;
-							LineStyle.Color = FLinearColor::Red;
-							LineStyle.Thickness = 1.0f;
-								
-							UDrawDebugLibrary::DrawDebugSimpleSphere(FDebugDrawer::MakeDebugDrawer(PDI), ContactState.Point[CharacterIdx][BoneIdx], FRotator::ZeroRotator, LineStyle, true, 5.0f);
-						}
-
-						FDrawDebugLineStyle LineStyle;
-						LineStyle.Color = FLinearColor::Green;
-						LineStyle.Thickness = 1.0f;
-
-						UDrawDebugLibrary::DrawDebugSimpleSphere(FDebugDrawer::MakeDebugDrawer(PDI), ContactState.Position[CharacterIdx][BoneIdx], FRotator::ZeroRotator, LineStyle, true, 5.0f);
-					}
-				}
-			}
-
-			// Debug Draw
-
-			if (ToolKit->Controller->ViewportSettings->bDrawDebugTesting)
-			{
-				for (int32 CharacterIdx = 0; CharacterIdx < CharacterNum; CharacterIdx++)
-				{
-					FLinearColor IdentifierColor = CharacterColors[CharacterIdx];
-					IdentifierColor.A = ToolKit->Controller->ViewportSettings->DrawRangeIdentifierOpacity;
-
-					PoseState.PoseIdx = CharacterIdx;
-
-					BehaviorPreviews[CharacterIdx]->OnPreviewDrawDebug(FDebugDrawer::MakeDebugDrawer(PDI), ControlObject, ControlObjectElements[CharacterIdx], PoseState, IdentifierColor);
-				}
-			}
-		}
-	}
-
-	void FControllerMode::Render(const FSceneView* View, FViewport* Viewport, FPrimitiveDrawInterface* PDI)
-	{
-		if (TSharedPtr<FControllerToolkit> ToolKit = WeakToolkit.Pin())
-		{
-			if (ToolKit->ControllerState == EControllerState::Stopped)
-			{
-				RenderTraining(View, Viewport, PDI);
-			}
-			else
-			{
-				RenderTesting(View, Viewport, PDI);
-			}
-		}
 
 		FEdMode::Render(View, Viewport, PDI);
 	}
diff --git a/Engine/Plugins/Experimental/Animation/AnimGen/Source/AnimGenEditor/Private/AnimGenEditorControllerToolkit.cpp b/Engine/Plugins/Experimental/Animation/AnimGen/Source/AnimGenEditor/Private/AnimGenEditorControllerToolkit.cpp
--- a/Engine/Plugins/Experimental/Animation/AnimGen/Source/AnimGenEditor/Private/AnimGenEditorControllerToolkit.cpp
+++ b/Engine/Plugins/Experimental/Animation/AnimGen/Source/AnimGenEditor/Private/AnimGenEditorControllerToolkit.cpp
@@ -493,8 +493,6 @@ namespace UE::AnimGen::Editor
 		IDetailCategoryBuilder& StatisticsCategory = DetailBuilder.EditCategory("Statistics", LOCTEXT("Statistics", "Statistics"), ECategoryPriority::Important);
 		StatisticsCategory.AddProperty(GET_MEMBER_NAME_CHECKED(UAnimGenController, ControlEncoderSize), UAnimGenController::StaticClass());
 		StatisticsCategory.AddProperty(GET_MEMBER_NAME_CHECKED(UAnimGenController, DenoiserSize), UAnimGenController::StaticClass());
-		StatisticsCategory.AddProperty(GET_MEMBER_NAME_CHECKED(UAnimGenController, ControlEncoderInferenceTime), UAnimGenController::StaticClass());
-		StatisticsCategory.AddProperty(GET_MEMBER_NAME_CHECKED(UAnimGenController, DenoiserInferenceTime), UAnimGenController::StaticClass());
 	}
 
 	class SControllerViewportToolBar;
@@ -1270,12 +1268,8 @@ namespace UE::AnimGen::Editor
 				Controller->ControlEncoderNetwork->Init(ControllerInputSize, ControllerOutputSize, ControllerCompatibilityHash, ControllerFileData);
 				ControllerFileData.Empty();
 
-				const int32 ControllerByteNum = Controller->ControlEncoderNetwork ? Controller->ControlEncoderNetwork->GetSnapshotByteNum() : 0;
-				ControlEncoderNetworkData = UE::Learning::SharedMemory::Allocate<1, uint8>({ ControllerByteNum });
-				if (Controller->ControlEncoderNetwork)
-				{
-					Controller->ControlEncoderNetwork->SaveToSnapshot(MakeArrayView<uint8>(ControlEncoderNetworkData.View.GetData(), ControlEncoderNetworkData.View.Num()));
-				}
+				ControlEncoderNetworkData = UE::Learning::SharedMemory::Allocate<1, uint8>({ Controller->ControlEncoderNetwork->GetSnapshotByteNum() });
+				Controller->ControlEncoderNetwork->SaveToSnapshot(MakeArrayView<uint8>(ControlEncoderNetworkData.View.GetData(), ControlEncoderNetworkData.View.Num()));
 
 				{
 					int32 LOD0CompatibilityHash = 0;
@@ -1304,12 +1298,8 @@ namespace UE::AnimGen::Editor
 					Controller->LOD0Network->Init(LOD0InputSize, LOD0OutputSize, LOD0CompatibilityHash, LOD0FileData);
 					LOD0FileData.Empty();
 
-					const int32 LOD0ByteNum = Controller->LOD0Network ? Controller->LOD0Network->GetSnapshotByteNum() : 0;
-					LOD0NetworkData = UE::Learning::SharedMemory::Allocate<1, uint8>({ LOD0ByteNum });
-					if (Controller->LOD0Network)
-					{
-						Controller->LOD0Network->SaveToSnapshot(MakeArrayView<uint8>(LOD0NetworkData.View.GetData(), LOD0NetworkData.View.Num()));
-					}
+					LOD0NetworkData = UE::Learning::SharedMemory::Allocate<1, uint8>({ Controller->LOD0Network->GetSnapshotByteNum() });
+					Controller->LOD0Network->SaveToSnapshot(MakeArrayView<uint8>(LOD0NetworkData.View.GetData(), LOD0NetworkData.View.Num()));
 				}
 
 				{
@@ -1339,12 +1329,8 @@ namespace UE::AnimGen::Editor
 					Controller->LOD1Network->Init(LOD1InputSize, LOD1OutputSize, LOD1CompatibilityHash, LOD1FileData);
 					LOD1FileData.Empty();
 
-					const int32 LOD1ByteNum = Controller->LOD1Network ? Controller->LOD1Network->GetSnapshotByteNum() : 0;
-					LOD1NetworkData = UE::Learning::SharedMemory::Allocate<1, uint8>({ LOD1ByteNum });
-					if (Controller->LOD1Network)
-					{
-						Controller->LOD1Network->SaveToSnapshot(MakeArrayView<uint8>(LOD1NetworkData.View.GetData(), LOD1NetworkData.View.Num()));
-					}
+					LOD1NetworkData = UE::Learning::SharedMemory::Allocate<1, uint8>({ Controller->LOD1Network->GetSnapshotByteNum() });
+					Controller->LOD1Network->SaveToSnapshot(MakeArrayView<uint8>(LOD1NetworkData.View.GetData(), LOD1NetworkData.View.Num()));
 				}
 
 				{
@@ -1374,14 +1360,18 @@ namespace UE::AnimGen::Editor
 					Controller->LOD2Network->Init(LOD2InputSize, LOD2OutputSize, LOD2CompatibilityHash, LOD2FileData);
 					LOD2FileData.Empty();
 
-					const int32 LOD2ByteNum = Controller->LOD2Network ? Controller->LOD2Network->GetSnapshotByteNum() : 0;
-					LOD2NetworkData = UE::Learning::SharedMemory::Allocate<1, uint8>({ LOD2ByteNum });
-					if (Controller->LOD2Network)
-					{
-						Controller->LOD2Network->SaveToSnapshot(MakeArrayView<uint8>(LOD2NetworkData.View.GetData(), LOD2NetworkData.View.Num()));
-					}
+					LOD2NetworkData = UE::Learning::SharedMemory::Allocate<1, uint8>({ Controller->LOD2Network->GetSnapshotByteNum() });
+					Controller->LOD2Network->SaveToSnapshot(MakeArrayView<uint8>(LOD2NetworkData.View.GetData(), LOD2NetworkData.View.Num()));
 				}
 
+				// Set Network Sizes
+
+				Controller->ControlEncoderSize = Controller->ControlEncoderNetwork->GetSnapshotByteNum() / 1000;
+				Controller->DenoiserSize = (
+					Controller->LOD0Network->GetSnapshotByteNum() +
+					Controller->LOD1Network->GetSnapshotByteNum() + 
+					Controller->LOD2Network->GetSnapshotByteNum()) / 1000;
+
 				// Allocate Control Data
 
 				// We allocate another array of "controls" which we use to synchronize various events between the python training process and the C++ side.
@@ -1415,7 +1405,7 @@ namespace UE::AnimGen::Editor
 				ConfigObject->SetNumberField(TEXT("DatabaseTotalFrameNum"), DatabaseTotalFrameNum);
 				ConfigObject->SetNumberField(TEXT("ControlTotalFrameNum"), ControlTotalFrameNum);
 
-				ConfigObject->SetNumberField(TEXT("ControllerByteNum"), ControllerByteNum);
+				ConfigObject->SetNumberField(TEXT("ControllerByteNum"), Controller->ControlEncoderNetwork->GetSnapshotByteNum());
 				ConfigObject->SetNumberField(TEXT("LOD0ByteNum"), Controller->LOD0Network->GetSnapshotByteNum());
 				ConfigObject->SetNumberField(TEXT("LOD1ByteNum"), Controller->LOD1Network->GetSnapshotByteNum());
 				ConfigObject->SetNumberField(TEXT("LOD2ByteNum"), Controller->LOD2Network->GetSnapshotByteNum());
@@ -1730,7 +1720,6 @@ namespace UE::AnimGen::Editor
 	static const FName ControllerTabViewport(TEXT("AnimGenEditorControllerViewportTabID"));
 	static const FName ControllerTabTimeline(TEXT("AnimGenEditorControllerTimelineTabID"));
 	static const FName ControllerTabTrainingSettings(TEXT("AnimGenEditorControllerTrainingSettingsTabID"));
-	static const FName ControllerTabTestingSettings(TEXT("AnimGenEditorControllerTestingSettingsTabID"));
 	static const FName ControllerTabTraining(TEXT("AnimGenEditorControllerTrainingTabID"));
 	static const FName ControllerTabControlObject(TEXT("AnimGenEditorControllerControlObjectTabID"));
 
@@ -1739,7 +1728,6 @@ namespace UE::AnimGen::Editor
 		Controller = InController;
 		Controller->TrainingSettings->LoadDatabaseAsync();
 		Controller->TrainingSettings->LoadAutoEncoderAsync();
-		Controller->TestingSettings->LoadDatabaseAsync();
 
 		PreviewScene = MakeShared<AnimDatabase::Editor::FPreviewScene>(AnimDatabase::Editor::FPreviewScene::ConstructionValues().ShouldSimulatePhysics(true).ForceUseMovementComponentInNonGameWorld(true));
 		//Temporary fix for missing attached assets - MDW (Copied from FPersonaToolkit::CreatePreviewScene)
@@ -1777,9 +1765,6 @@ namespace UE::AnimGen::Editor
 		TrainingSettingsWidget = PropertyModule.CreateDetailView(Args);
 		TrainingSettingsWidget->SetObject(InController->TrainingSettings);
 
-		TestingSettingsWidget = PropertyModule.CreateDetailView(Args);
-		TestingSettingsWidget->SetObject(InController->TestingSettings);
-
 		TrainingWidget = SNew(STraining, TrainingModel.ToWeakPtr());
 		ControlObjectWidget = SNew(SControlObject);
 
@@ -1799,7 +1784,6 @@ namespace UE::AnimGen::Editor
 							FTabManager::NewStack()
 							->SetSizeCoefficient(0.8f)
 							->AddTab(ControllerTabTrainingSettings, ETabState::OpenedTab)
-							->AddTab(ControllerTabTestingSettings, ETabState::OpenedTab)
 							->SetForegroundTab(ControllerTabTrainingSettings)
 						)
 						->Split
@@ -1859,74 +1843,6 @@ namespace UE::AnimGen::Editor
 			bCreateDefaultToolbarParam, 
 			InController,
 			bIsToolbarFocusableParam);
-
-		TSharedPtr<FExtender> ToolbarExtender = MakeShareable(new FExtender);
-		ToolbarExtender->AddToolBarExtension(
-			"Asset",
-			EExtensionHook::After,
-			GetToolkitCommands(), // Optional: command list
-			FToolBarExtensionDelegate::CreateRaw(this, &FControllerToolkit::AddToolbarButton));
-
-		AddToolbarExtender(ToolbarExtender);
-
-		RegenerateMenusAndToolbars();
-	}
-
-	void FControllerToolkit::AddToolbarButton(FToolBarBuilder& Builder)
-	{
-		Builder.BeginSection("PlayControls");
-
-		if (ControllerState == EControllerState::Playing)
-		{
-			Builder.AddToolBarButton(
-				FUIAction(
-					FExecuteAction::CreateLambda([this]() { ControllerState = EControllerState::Paused; RegenerateMenusAndToolbars(); }),
-					FCanExecuteAction::CreateLambda([this]() { return ControllerState == EControllerState::Playing; })),
-				NAME_None,
-				FText(),
-				LOCTEXT("Pause_Tooltip", "Pause"),
-				FSlateIcon(FAppStyle::GetAppStyleSetName(), "PlayWorld.PausePlaySession"));
-		}
-		else
-		{
-			Builder.AddToolBarButton(
-				FUIAction(
-					FExecuteAction::CreateLambda([this]() { bResetRequested = true; ControllerState = EControllerState::Playing; RegenerateMenusAndToolbars(); TimelineWidget->SetEnabled(false); TrainingSettingsWidget->SetEnabled(false); }),
-					FCanExecuteAction::CreateLambda([this]() { return ControllerState == EControllerState::Paused || (ControllerState == EControllerState::Stopped && Controller->IsValid()); })),
-				NAME_None,
-				FText(),
-				LOCTEXT("Play_Tooltip", "Play this controller in the editor"),
-				FSlateIcon(FAppStyle::GetAppStyleSetName(), "PlayWorld.PlayInViewport"));
-		}
-
-		Builder.AddToolBarButton(
-			FUIAction(
-				FExecuteAction::CreateLambda([this]() { ControllerState = EControllerState::Stepping; RegenerateMenusAndToolbars(); }),
-				FCanExecuteAction::CreateLambda([this]() { return ControllerState == EControllerState::Paused; })),
-			NAME_None,
-			FText(),
-			LOCTEXT("StepFrame_Tooltip", "Advances a single frame"),
-			FSlateIcon(FAppStyle::GetAppStyleSetName(), "PlayWorld.SingleFrameAdvance"));
-
-		Builder.AddToolBarButton(
-			FUIAction(
-				FExecuteAction::CreateLambda([this]() { ControllerState = EControllerState::Stopped; RegenerateMenusAndToolbars(); TimelineWidget->SetEnabled(true); TrainingSettingsWidget->SetEnabled(true); }),
-				FCanExecuteAction::CreateLambda([this]() { return ControllerState != EControllerState::Stopped; })),
-			NAME_None,
-			FText(),
-			LOCTEXT("Stop_Tooltip", "Stop"),
-			FSlateIcon(FAppStyle::GetAppStyleSetName(), "PlayWorld.StopPlaySession"));
-
-		Builder.AddToolBarButton(
-			FUIAction(
-				FExecuteAction::CreateLambda([this]() { bResetRequested = true; RegenerateMenusAndToolbars(); }),
-				FCanExecuteAction::CreateLambda([this]() { return ControllerState != EControllerState::Stopped && !bResetRequested; })),
-			NAME_None,
-			FText(),
-			LOCTEXT("Reset_Tooltip", "Reset"),
-			FSlateIcon(FAppStyle::GetAppStyleSetName(), "ContentBrowser.ResetPrimitiveToDefault"));
-
-		Builder.EndSection();
 	}
 
 	void FControllerToolkit::RegisterTabSpawners(const TSharedRef<FTabManager>& InTabManager)
@@ -1972,12 +1888,6 @@ namespace UE::AnimGen::Editor
 			.SetGroup(WorkspaceMenuCategoryRef)
 			.SetIcon(FSlateIcon(FAppStyle::GetAppStyleSetName(), "LevelEditor.Tabs.Details"));
 
-		InTabManager->RegisterTabSpawner(ControllerTabTestingSettings, FOnSpawnTab::CreateLambda([this](const FSpawnTabArgs&) {
-			return SNew(SDockTab).Label(LOCTEXT("TestingSettingsTab_Title", "Testing Settings"))[TestingSettingsWidget.ToSharedRef()]; }))
-			.SetDisplayName(LOCTEXT("TestingSettingsTab", "Testing Settings"))
-			.SetGroup(WorkspaceMenuCategoryRef)
-			.SetIcon(FSlateIcon(FAppStyle::GetAppStyleSetName(), "LevelEditor.Tabs.Details"));
-
 		InTabManager->RegisterTabSpawner(ControllerTabTraining, FOnSpawnTab::CreateLambda([this](const FSpawnTabArgs&) {
 			return SNew(SDockTab).Label(LOCTEXT("TrainingTab_Title", "Training"))[TrainingWidget.ToSharedRef()]; }))
 			.SetDisplayName(LOCTEXT("TrainingTab", "Training"))
@@ -2001,7 +1911,6 @@ namespace UE::AnimGen::Editor
 		InTabManager->UnregisterTabSpawner(ControllerTabPreviewSettings);
 		InTabManager->UnregisterTabSpawner(ControllerTabTimeline);
 		InTabManager->UnregisterTabSpawner(ControllerTabTrainingSettings);
-		InTabManager->UnregisterTabSpawner(ControllerTabTestingSettings);
 		InTabManager->UnregisterTabSpawner(ControllerTabTraining);
 		InTabManager->UnregisterTabSpawner(ControllerTabControlObject);
 	}
@@ -2013,7 +1922,7 @@ namespace UE::AnimGen::Editor
 
 	FText FControllerToolkit::GetBaseToolkitName() const
 	{
-		return LOCTEXT("AnimGenEditorControllerAppLabel", "Animation Controller Editor");
+		return LOCTEXT("AnimGenEditorControllerAppLabel", "AnimGen Controller Editor");
 	}
 
 	FText FControllerToolkit::GetToolkitName() const
diff --git a/Engine/Plugins/Experimental/Animation/AnimGen/Source/AnimGenEditor/Public/AnimGenEditorControllerMode.h b/Engine/Plugins/Experimental/Animation/AnimGen/Source/AnimGenEditor/Public/AnimGenEditorControllerMode.h
--- a/Engine/Plugins/Experimental/Animation/AnimGen/Source/AnimGenEditor/Public/AnimGenEditorControllerMode.h
+++ b/Engine/Plugins/Experimental/Animation/AnimGen/Source/AnimGenEditor/Public/AnimGenEditorControllerMode.h
@@ -38,24 +38,11 @@ namespace UE::AnimGen::Editor
 		virtual void Tick(FEditorViewportClient* ViewportClient, float DeltaTime) override;
 		virtual void Render(const FSceneView* View, FViewport* Viewport, FPrimitiveDrawInterface* PDI);
 
-		virtual bool MouseMove(FEditorViewportClient* ViewportClient, FViewport* Viewport, int32 X, int32 Y) override;
-		virtual bool CapturedMouseMove(FEditorViewportClient* ViewportClient, FViewport* Viewport, int32 X, int32 Y) override;
-		virtual bool InputAxis(FEditorViewportClient* InViewportClient, FViewport* Viewport, int32 ControllerId, FKey Key, float Delta, float DeltaTime) override;
-		virtual bool InputKey(FEditorViewportClient* ViewportClient, FViewport* Viewport, FKey Key, EInputEvent Event) override;
-
-		void TickTraining(FEditorViewportClient* EditorViewportClient, float DeltaTime);
-		void TickTesting(FEditorViewportClient* EditorViewportClient, float DeltaTime);
-
-		void RenderTraining(const FSceneView* View, FViewport* Viewport, FPrimitiveDrawInterface* PDI);
-		void RenderTesting(const FSceneView* View, FViewport* Viewport, FPrimitiveDrawInterface* PDI);
-
 	private:
 
 		/** Maximum number of characters to draw in the scene */
 		const int32 MaxCharacterNum = 128;
 
-		bool bCaptureInput = false;
-
 		int32 CharacterNum = 0;
 		TArray<FLinearColor> CharacterColors;
 
@@ -87,57 +74,6 @@ namespace UE::AnimGen::Editor
 		/** Frame Attributes database hash */
 		int32 FrameAttributesDatabaseContentHash = 0;
 
-		/** By default we should reset at the start */
-		bool bReset = true;
-
-		/** Trained Database content hash */
-		int32 TrainedDatabaseContentHash = 0;
-		UObject* TrainedDatabaseFrameRangesDefaultObject = nullptr;
-		int32 TrainedAutoEncoderContentHash = 0;
-
-		TObjectPtr<UAnimGenBehaviorPreview> ActiveBehaviorObject = nullptr;
-		TArray<UAnimGenBehaviorPreview*> BehaviorPreviews;
-
-		TSharedPtr<UE::Learning::FNeuralNetworkInference> EncoderInference;
-		TSharedPtr<UE::Learning::FNeuralNetworkInference> DecoderInference;
-		TSharedPtr<UE::Learning::FNeuralNetworkInference> ControlEncoderInference;
-		TSharedPtr<UE::Learning::FNeuralNetworkInference> LOD0Inference;
-		TSharedPtr<UE::Learning::FNeuralNetworkInference> LOD1Inference;
-		TSharedPtr<UE::Learning::FNeuralNetworkInference> LOD2Inference;
-
-		TObjectPtr<ULearningNeuralNetworkData> EncoderInferenceNetwork;
-		TObjectPtr<ULearningNeuralNetworkData> DecoderInferenceNetwork;
-		TObjectPtr<ULearningNeuralNetworkData> ControlEncoderInferenceNetwork;
-		TObjectPtr<ULearningNeuralNetworkData> LOD0InferenceNetwork;
-		TObjectPtr<ULearningNeuralNetworkData> LOD1InferenceNetwork;
-		TObjectPtr<ULearningNeuralNetworkData> LOD2InferenceNetwork;
-
-		UE::AnimGen::Editor::FPerfCounter ControlEncoderPerfCounter;
-		UE::AnimGen::Editor::FPerfCounter DenoiserPerfCounter;
-
-		TArray<TPair<FKey, EInputEvent>, TInlineAllocator<8>> KeyEvents;
-
 		FAnimDatabasePoseState PoseState;
-
-		TLearningArray<2, float> MemoryVector;
-		TLearningArray<2, float> ControlVector;
-		TLearningArray<2, float> PoseVector;
-		TLearningArray<2, float> UnnormalizedEncodedPoseVector;
-		TLearningArray<2, float> EncodedPoseVector;
-		TLearningArray<2, float> NextEncodedPoseVector;
-		TLearningArray<2, float> ControllerInputVector;
-		TLearningArray<2, float> ControllerOutputVector;
-		TLearningArray<2, float> DenoiserInputVector;
-
-		UE::AnimDatabase::FPoseData OutputPoseData;
-		TLearningArray<2, float> PredictedPoseVector;
-		UE::AnimDatabase::FPoseData AdjustedPoseData;
-		UE::AnimDatabase::FPoseGlobalBoneData AdjustedGlobalBoneData;
-		UE::AnimDatabase::PoseData::ContactPinning::FContactState ContactState;
-
-		uint32 RandomState = 1234;
-
-		FAnimGenControlObject ControlObject;
-		TArray<FAnimGenControlObjectElement> ControlObjectElements;
 	};
 }
diff --git a/Engine/Plugins/Experimental/Animation/AnimGen/Source/AnimGenEditor/Public/AnimGenEditorControllerToolkit.h b/Engine/Plugins/Experimental/Animation/AnimGen/Source/AnimGenEditor/Public/AnimGenEditorControllerToolkit.h
--- a/Engine/Plugins/Experimental/Animation/AnimGen/Source/AnimGenEditor/Public/AnimGenEditorControllerToolkit.h
+++ b/Engine/Plugins/Experimental/Animation/AnimGen/Source/AnimGenEditor/Public/AnimGenEditorControllerToolkit.h
@@ -100,14 +100,6 @@ namespace UE::AnimGen::Editor
 		virtual void CustomizeDetails(IDetailLayoutBuilder& DetailBuilder) override;
 	};
 
-	enum class EControllerState : uint8
-	{
-		Stopped,
-		Playing,
-		Paused,
-		Stepping,
-	};
-
 	class FControllerToolkit : public FAssetEditorToolkit, public FTickableEditorObject
 	{
 	public:
@@ -137,12 +129,6 @@ namespace UE::AnimGen::Editor
 		/** Model for the tracks Widget */
 		TSharedPtr<AnimDatabase::Editor::FTimelineTracksModel> TracksModel;
 
-		/** State of the controller */
-		EControllerState ControllerState = EControllerState::Stopped;
-
-		/** If a reset is requested */
-		bool bResetRequested = false;
-
 		/** Controller Training model */
 		TSharedPtr<FControllerTrainingModel> TrainingModel;
 
@@ -167,9 +153,6 @@ namespace UE::AnimGen::Editor
 		/** Training Settings Tab Widget */
 		TSharedPtr<IDetailsView> TrainingSettingsWidget;
 
-		/** Testing Settings Tab Widget */
-		TSharedPtr<IDetailsView> TestingSettingsWidget;
-
 		/* Training Widget */
 		TSharedPtr<SWidget> TrainingWidget;
 
