AnimGen: Force UI Refresh after starting training

#Synced-CL 52479696

diff --git a/Engine/Plugins/Experimental/Animation/AnimGen/Source/AnimGen/Private/AnimGenController.cpp b/Engine/Plugins/Experimental/Animation/AnimGen/Source/AnimGen/Private/AnimGenController.cpp
--- a/Engine/Plugins/Experimental/Animation/AnimGen/Source/AnimGen/Private/AnimGenController.cpp
+++ b/Engine/Plugins/Experimental/Animation/AnimGen/Source/AnimGen/Private/AnimGenController.cpp
@@ -197,6 +197,15 @@ bool UAnimGenControllerTrainingSettings::Update()
 				if (ControlSchema.ObservationSchema->Schema.IsValid(ControlSchemaElement.SchemaElement))
 				{
 					Behavior->MakeControlSets(ControlSets, ControlObject, DatabaseObject, QueryRanges);
+					const int32 ControlSetNum = ControlSets.Num();
+					for (int32 ControlSetIdx = 0; ControlSetIdx < ControlSetNum; ControlSetIdx++)
+					{
+						if (!ControlSets[ControlSetIdx].IsValid())
+						{
+							ControlSets.Empty();
+							break;
+						}
+					}
 				}
 				else
 				{
diff --git a/Engine/Plugins/Experimental/Animation/AnimGen/Source/AnimGenEditor/Private/AnimGenEditorAutoEncoderToolkit.cpp b/Engine/Plugins/Experimental/Animation/AnimGen/Source/AnimGenEditor/Private/AnimGenEditorAutoEncoderToolkit.cpp
--- a/Engine/Plugins/Experimental/Animation/AnimGen/Source/AnimGenEditor/Private/AnimGenEditorAutoEncoderToolkit.cpp
+++ b/Engine/Plugins/Experimental/Animation/AnimGen/Source/AnimGenEditor/Private/AnimGenEditorAutoEncoderToolkit.cpp
@@ -850,8 +850,6 @@ namespace UE::AnimGen::Editor
 
 				AutoEncoder->TrainedContentHash = UAnimDatabaseFrameRangesLibrary::FrameRangesContentHash(TrainingDatabase, TrainingFrameRanges);
 
-				ToolkitPtr->EditingAssetWidget->ForceRefresh();
-
 				// Extract Animation Data
 
 				TrainingDatabase->WaitForCompressionOnAnimSequencesFromArrayView(RangeSequenceIndices);
@@ -1165,6 +1163,12 @@ namespace UE::AnimGen::Editor
 				AutoEncoder->DecoderNetwork->SaveToSnapshot(MakeArrayView<uint8>(DecoderNetworkData.View.GetData(), DecoderNetworkData.View.Num()));
 				AutoEncoder->DecoderSize = AutoEncoder->DecoderNetwork->GetSnapshotByteNum() / 1000;
 
+				// Refresh UI
+
+				AutoEncoder->TrainingSettings->Update();
+				ToolkitPtr->TrainingSettingsWidget->ForceRefresh();
+				ToolkitPtr->EditingAssetWidget->ForceRefresh();
+
 				// Write Config File
 
 				const FString EnginePath = FPaths::EngineDir();
diff --git a/Engine/Plugins/Experimental/Animation/AnimGen/Source/AnimGenEditor/Private/AnimGenEditorControllerToolkit.cpp b/Engine/Plugins/Experimental/Animation/AnimGen/Source/AnimGenEditor/Private/AnimGenEditorControllerToolkit.cpp
--- a/Engine/Plugins/Experimental/Animation/AnimGen/Source/AnimGenEditor/Private/AnimGenEditorControllerToolkit.cpp
+++ b/Engine/Plugins/Experimental/Animation/AnimGen/Source/AnimGenEditor/Private/AnimGenEditorControllerToolkit.cpp
@@ -1372,6 +1372,12 @@ namespace UE::AnimGen::Editor
 					Controller->LOD1Network->GetSnapshotByteNum() + 
 					Controller->LOD2Network->GetSnapshotByteNum()) / 1000;
 
+                // Update UI
+
+				Controller->TrainingSettings->Update();
+				ToolkitPtr->TrainingSettingsWidget->ForceRefresh();
+				ToolkitPtr->EditingAssetWidget->ForceRefresh();
+
 				// Allocate Control Data
 
 				// We allocate another array of "controls" which we use to synchronize various events between the python training process and the C++ side.
