AnimGen: Added Encoded Behavior

#Synced-CL 52628728

diff --git a/Engine/Plugins/Experimental/Animation/AnimGen/Source/AnimGen/Private/AnimGenBehavior.cpp b/Engine/Plugins/Experimental/Animation/AnimGen/Source/AnimGen/Private/AnimGenBehavior.cpp
--- a/Engine/Plugins/Experimental/Animation/AnimGen/Source/AnimGen/Private/AnimGenBehavior.cpp
+++ b/Engine/Plugins/Experimental/Animation/AnimGen/Source/AnimGen/Private/AnimGenBehavior.cpp
@@ -428,6 +428,64 @@ void UAnimGenBehavior_Tagged::DrawDebugControl_Implementation(
 	}
 }
 
+FAnimGenControlSchemaElement UAnimGenBehavior_Encoded::SpecifyControl_Implementation(FAnimGenControlSchema& InControlSchema) const
+{
+
+	return UAnimGenControls::SpecifyEncodingControl(
+		InControlSchema,
+		Behavior ? Behavior->SpecifyControl(InControlSchema) : UAnimGenControls::SpecifyNullControl(InControlSchema),
+		EncodingSize,
+		HiddenLayerNum,
+		ActivationFunction);
+}
+
+void UAnimGenBehavior_Encoded::MakeControlSets_Implementation(TArray<FAnimGenControlSet>& OutControlSets, FAnimGenControlObject& InControlObject, const UAnimDatabase* InDatabase, const FAnimDatabaseFrameRanges& InFrameRanges) const
+{
+	TArray<FAnimGenControlSet> SubControlSets;
+	if (Behavior)
+	{
+		Behavior->MakeControlSets(SubControlSets, InControlObject, InDatabase, InFrameRanges);
+	}
+
+	for (const FAnimGenControlSet& SubControlSet : SubControlSets)
+	{
+		OutControlSets.Add(UAnimGenControls::MakeEncodingControlSet(InControlObject, SubControlSet));
+	}
+}
+
+void UAnimGenBehavior_Encoded::DrawDebugControl_Implementation(
+	const FDebugDrawer& Drawer,
+	const FAnimGenControlObject& InControlObject,
+	const FAnimGenControlObjectElement InControlObjectElement,
+	const FAnimDatabasePoseState& InPoseState,
+	const UAnimDatabase* InDatabase,
+	const UAnimGenAutoEncoder* InAutoEncoder,
+	const int32 SequenceIdx,
+	const float SequenceTime,
+	const int32 RangeStart,
+	const int32 RangeLength,
+	const FLinearColor& IdentifierColor) const
+{
+	FAnimGenControlObjectElement SubControl;
+	UAnimGenControls::GetEncodingControl(SubControl, InControlObject, InControlObjectElement);
+
+	if (Behavior)
+	{
+		Behavior->DrawDebugControl(
+			Drawer,
+			InControlObject,
+			SubControl,
+			InPoseState,
+			InDatabase,
+			InAutoEncoder,
+			SequenceIdx,
+			SequenceTime,
+			RangeStart,
+			RangeLength,
+			IdentifierColor);
+	}
+}
+
 FAnimGenControlSchemaElement UAnimGenBehavior_Multi::SpecifyControl_Implementation(FAnimGenControlSchema& InControlSchema) const
 {
 	TArray<FName, TInlineAllocator<32>> SubBehaviorNames;
@@ -623,6 +681,13 @@ FAnimGenControlObjectElement UAnimGenBehaviorLibrary::MakeTaggedBehaviorControlF
 		});
 }
 
+FAnimGenControlObjectElement UAnimGenBehaviorLibrary::MakeEncodedBehaviorControl(
+	FAnimGenControlObject& InControlObject,
+	const FAnimGenControlObjectElement& BehaviorControlElement)
+{
+	return UAnimGenControls::MakeEncodingControl(InControlObject, BehaviorControlElement);
+}
+
 FAnimGenControlObjectElement UAnimGenBehaviorLibrary::MakeMultiBehaviorControl(
 	FAnimGenControlObject& InControlObject,
 	const FName Behavior,
diff --git a/Engine/Plugins/Experimental/Animation/AnimGen/Source/AnimGen/Public/AnimGenBehavior.h b/Engine/Plugins/Experimental/Animation/AnimGen/Source/AnimGen/Public/AnimGenBehavior.h
--- a/Engine/Plugins/Experimental/Animation/AnimGen/Source/AnimGen/Public/AnimGenBehavior.h
+++ b/Engine/Plugins/Experimental/Animation/AnimGen/Source/AnimGen/Public/AnimGenBehavior.h
@@ -315,6 +315,58 @@ public:
 #endif
 };
 
+/** Behavior that encodes the controls of some sub-behavior. */
+UCLASS(BlueprintType, Blueprintable, meta = (DisplayName = "Encoded Behavior"))
+class UAnimGenBehavior_Encoded : public UAnimGenBehavior
+{
+	GENERATED_BODY()
+
+#if WITH_EDITORONLY_DATA
+
+public:
+
+	/** Sub-behavior */
+	UPROPERTY(EditAnywhere, Instanced, Category = "Settings")
+	TObjectPtr<UAnimGenBehavior> Behavior;
+
+	/** Size to encode the sub-behaviors controls to */
+	UPROPERTY(EditAnywhere, Category = "Settings", meta = (UIMin = "1", ClampMin = "1"))
+	int32 EncodingSize = 128;
+	
+	/** Number of layers to use to encode the sub-behavior controls */
+	UPROPERTY(EditAnywhere, Category = "Settings", meta = (UIMin = "1", ClampMin = "1"))
+	int32 HiddenLayerNum = 1;
+	
+	/** Activation function to use to encode the sub-behavior controls */
+	UPROPERTY(EditAnywhere, Category = "Settings")
+	EAnimGenActivationFunction ActivationFunction = EAnimGenActivationFunction::ELU;
+
+#endif
+
+#if WITH_EDITOR
+
+public:
+
+	UE_API virtual FAnimGenControlSchemaElement SpecifyControl_Implementation(FAnimGenControlSchema& InControlSchema) const override;
+
+	UE_API virtual void MakeControlSets_Implementation(TArray<FAnimGenControlSet>& OutControlSets, FAnimGenControlObject& InControlObject, const UAnimDatabase* InDatabase, const FAnimDatabaseFrameRanges& InFrameRanges) const override;
+
+	UE_API virtual void DrawDebugControl_Implementation(
+		const FDebugDrawer& Drawer,
+		const FAnimGenControlObject& InControlObject,
+		const FAnimGenControlObjectElement InControlObjectElement,
+		const FAnimDatabasePoseState& InPoseState,
+		const UAnimDatabase* InDatabase,
+		const UAnimGenAutoEncoder* InAutoEncoder,
+		const int32 SequenceIdx,
+		const float SequenceTime,
+		const int32 RangeStart,
+		const int32 RangeLength,
+		const FLinearColor& IdentifierColor) const override;
+
+#endif
+};
+
 /** Entry specifying a behavior */
 USTRUCT(BlueprintType)
 struct FAnimGenBehaviorEntry
@@ -445,6 +497,12 @@ public:
 		const FName Behavior,
 		const FAnimGenControlObjectElement& BehaviorControlElement);
 
+	/** Make a control object element for the Encoded Behavior */
+	UFUNCTION(BlueprintPure, Category = "AnimGen")
+	static UE_API FAnimGenControlObjectElement MakeEncodedBehaviorControl(
+		UPARAM(ref) FAnimGenControlObject& InControlObject,
+		const FAnimGenControlObjectElement& BehaviorControlElement);
+
 public:
 
 	/** Projects a location onto the floor plane given by the provided height */
