AnimGen: Minor bug fixes

#Synced-CL 52698189

diff --git a/Engine/Plugins/Experimental/Animation/AnimGen/Source/AnimGen/Private/AnimGenAutoEncoder.cpp b/Engine/Plugins/Experimental/Animation/AnimGen/Source/AnimGen/Private/AnimGenAutoEncoder.cpp
--- a/Engine/Plugins/Experimental/Animation/AnimGen/Source/AnimGen/Private/AnimGenAutoEncoder.cpp
+++ b/Engine/Plugins/Experimental/Animation/AnimGen/Source/AnimGen/Private/AnimGenAutoEncoder.cpp
@@ -14,6 +14,8 @@
 #include "UObject/Package.h"
 #include "Misc/ScopedSlowTask.h"
 
+#include "LearningNeuralNetwork.h"
+
 #include UE_INLINE_GENERATED_CPP_BY_NAME(AnimGenAutoEncoder)
 
 #define LOCTEXT_NAMESPACE "AnimDatabase"
@@ -264,7 +266,10 @@ UAnimGenAutoEncoder::UAnimGenAutoEncoder(const FObjectInitializer& ObjectInitial
 
 bool UAnimGenAutoEncoder::IsValid() const
 {
-	return !BoneNames.IsEmpty() && EncoderNetwork && DecoderNetwork;
+	return 
+		!BoneNames.IsEmpty() && 
+		EncoderNetwork && EncoderNetwork->GetUpdateNumber() > 0 &&
+		DecoderNetwork && DecoderNetwork->GetUpdateNumber() > 0;
 }
 
 void UAnimGenAutoEncoder::Invalidate()
diff --git a/Engine/Plugins/Experimental/Animation/AnimGen/Source/AnimGen/Private/AnimGenController.cpp b/Engine/Plugins/Experimental/Animation/AnimGen/Source/AnimGen/Private/AnimGenController.cpp
--- a/Engine/Plugins/Experimental/Animation/AnimGen/Source/AnimGen/Private/AnimGenController.cpp
+++ b/Engine/Plugins/Experimental/Animation/AnimGen/Source/AnimGen/Private/AnimGenController.cpp
@@ -14,6 +14,8 @@
 
 #include "Animation/Skeleton.h"
 
+#include "LearningNeuralNetwork.h"
+
 #include "DrawDebugLibrary.h"
 
 #include UE_INLINE_GENERATED_CPP_BY_NAME(AnimGenController)
@@ -271,10 +273,10 @@ bool UAnimGenController::IsValid() const
 		AutoEncoder->IsValid() &&
 		ControlSchema.IsValid() &&
 		ControlSchema.IsElementValid(ControlSchemaElement) &&
-		ControlEncoderNetwork &&
-		LOD0Network &&
-		LOD1Network &&
-		LOD2Network;
+		ControlEncoderNetwork && ControlEncoderNetwork->GetUpdateNumber() > 0 &&
+		LOD0Network && LOD0Network->GetUpdateNumber() > 0 &&
+		LOD1Network && LOD1Network->GetUpdateNumber() > 0 &&
+		LOD2Network && LOD2Network->GetUpdateNumber() > 0;
 }
 
 void UAnimGenController::Invalidate()
diff --git a/Engine/Plugins/Experimental/Animation/AnimGen/Source/AnimGenEditor/Private/AnimGenEditorAutoEncoderMode.cpp b/Engine/Plugins/Experimental/Animation/AnimGen/Source/AnimGenEditor/Private/AnimGenEditorAutoEncoderMode.cpp
--- a/Engine/Plugins/Experimental/Animation/AnimGen/Source/AnimGenEditor/Private/AnimGenEditorAutoEncoderMode.cpp
+++ b/Engine/Plugins/Experimental/Animation/AnimGen/Source/AnimGenEditor/Private/AnimGenEditorAutoEncoderMode.cpp
@@ -747,6 +747,7 @@ namespace UE::AnimGen::Editor
 			// Integrate root
 
 			AutoEncoder->DenormalizePoseVectors(ReconstructedPoseVector);
+			AutoEncoder->ClampPoseVectors(ReconstructedPoseVector);
 
 			if (AutoEncoder->ViewportSettings->bIntegrateRootOnReconstruction && EditorToolKit->TimelineModel->GetTransportPlaybackMode() == EPlaybackMode::Type::PlayingForward)
 			{
