AnimGen: Small fix for edge-case crash in AnimGenController

#Synced-CL 52734312

diff --git a/Engine/Plugins/Experimental/Animation/AnimGen/Source/AnimGen/Private/AnimNode_AnimGenController.cpp b/Engine/Plugins/Experimental/Animation/AnimGen/Source/AnimGen/Private/AnimNode_AnimGenController.cpp
--- a/Engine/Plugins/Experimental/Animation/AnimGen/Source/AnimGen/Private/AnimNode_AnimGenController.cpp
+++ b/Engine/Plugins/Experimental/Animation/AnimGen/Source/AnimGen/Private/AnimNode_AnimGenController.cpp
@@ -58,7 +58,16 @@ void FAnimNode_AnimGenController::Update_AnyThread(const FAnimationUpdateContext
 
 	// Controller has changed so reset internal state
 
-	if (Controller != ActiveController || !PoseState.IsValid())
+	if (Controller != ActiveController || 
+		!PoseState.IsValid() || 
+		PoseVectorSize != ActiveController->AutoEncoder->GetPoseVectorSize() ||
+		EncodedPoseVectorSize != ActiveController->AutoEncoder->GetEncodingSize() ||
+		ControlSchemaHash != UE::Learning::Observation::GetSchemaObjectsCompatibilityHash(
+			ActiveController->ControlSchema.ObservationSchema->Schema, 
+			ActiveController->ControlSchemaElement.SchemaElement) ||
+		ControlVectorSize != ActiveController->ControlSchema.ObservationSchema->Schema.GetObservationVectorSize(ActiveController->ControlSchemaElement.SchemaElement) ||
+		EncodedControlVectorSize != ActiveController->ControlSchema.ObservationSchema->Schema.GetEncodedVectorSize(ActiveController->ControlSchemaElement.SchemaElement) ||
+		ControlDistributionVectorSize != ActiveController->ControlSchema.ObservationSchema->Schema.GetObservationDistributionVectorSize(ActiveController->ControlSchemaElement.SchemaElement))
 	{
 		// Update Active Controller
 
@@ -70,6 +79,7 @@ void FAnimNode_AnimGenController::Update_AnyThread(const FAnimationUpdateContext
 		const FAnimGenControlSchemaElement& ControlSchemaElement = ActiveController->ControlSchemaElement;
 		check(ControlSchema.ObservationSchema.IsValid() && ControlSchema.ObservationSchema->Schema.IsValid(ControlSchemaElement.SchemaElement));
 
+		ControlSchemaHash = UE::Learning::Observation::GetSchemaObjectsCompatibilityHash(ControlSchema.ObservationSchema->Schema, ControlSchemaElement.SchemaElement);
 		ControlVectorSize = ControlSchema.ObservationSchema->Schema.GetObservationVectorSize(ControlSchemaElement.SchemaElement);
 		EncodedControlVectorSize = ControlSchema.ObservationSchema->Schema.GetEncodedVectorSize(ControlSchemaElement.SchemaElement);
 		ControlDistributionVectorSize = ControlSchema.ObservationSchema->Schema.GetObservationDistributionVectorSize(ControlSchemaElement.SchemaElement);
diff --git a/Engine/Plugins/Experimental/Animation/AnimGen/Source/AnimGen/Public/AnimNode_AnimGenController.h b/Engine/Plugins/Experimental/Animation/AnimGen/Source/AnimGen/Public/AnimNode_AnimGenController.h
--- a/Engine/Plugins/Experimental/Animation/AnimGen/Source/AnimGen/Public/AnimNode_AnimGenController.h
+++ b/Engine/Plugins/Experimental/Animation/AnimGen/Source/AnimGen/Public/AnimNode_AnimGenController.h
@@ -123,6 +123,7 @@ private:
 	// Control Variables
 
 	bool bControlsValid = false;
+	int32 ControlSchemaHash = 0;
 	int32 ControlVectorSize = 0;
 	int32 EncodedControlVectorSize = 0;
 	int32 ControlDistributionVectorSize = 0;
